Adding a new experimental plugin to access platform-native DLC, if supported

- refactoring existing platform-native implementations to conform to the new API
- fixing a handful of bugs along the way

#Synced-CL 51451766

diff --git a/Engine/Plugins/Experimental/Microsoft/GDKDLCSupport/GDKDLCSupport.uplugin b/Engine/Plugins/Experimental/Microsoft/GDKDLCSupport/GDKDLCSupport.uplugin
deleted file mode 100644
diff --git a/Engine/Plugins/Experimental/Microsoft/GDKDLCSupport/Source/GDKDLCSupport/GDKDLCSupport.Build.cs b/Engine/Plugins/Experimental/Microsoft/GDKDLCSupport/Source/GDKDLCSupport/GDKDLCSupport.Build.cs
deleted file mode 100644
diff --git a/Engine/Plugins/Experimental/Microsoft/GDKDLCSupport/Source/GDKDLCSupport/Private/GDKDLCHandler.cpp b/Engine/Plugins/Experimental/Microsoft/GDKDLCSupport/Source/GDKDLCSupport/Private/GDKDLCHandler.cpp
deleted file mode 100644
diff --git a/Engine/Plugins/Experimental/Microsoft/GDKDLCSupport/Source/GDKDLCSupport/Private/GDKDLCHandler.h b/Engine/Plugins/Experimental/Microsoft/GDKDLCSupport/Source/GDKDLCSupport/Private/GDKDLCHandler.h
deleted file mode 100644
diff --git a/Engine/Plugins/Experimental/Microsoft/GDKDLCSupport/Source/GDKDLCSupport/Private/GDKDLCSupportModule.cpp b/Engine/Plugins/Experimental/Microsoft/GDKDLCSupport/Source/GDKDLCSupport/Private/GDKDLCSupportModule.cpp
deleted file mode 100644
diff --git a/Engine/Plugins/Experimental/Microsoft/GDKDLCSupport/Source/GDKDLCSupport/Public/GDKDLCSupportModule.h b/Engine/Plugins/Experimental/Microsoft/GDKDLCSupport/Source/GDKDLCSupport/Public/GDKDLCSupportModule.h
deleted file mode 100644
diff --git a/Engine/Plugins/Experimental/Microsoft/GDKPlatformDLC/GDKPlatformDLC.uplugin b/Engine/Plugins/Experimental/Microsoft/GDKPlatformDLC/GDKPlatformDLC.uplugin
new file mode 100644
--- /dev/null
+++ b/Engine/Plugins/Experimental/Microsoft/GDKPlatformDLC/GDKPlatformDLC.uplugin
@@ -0,0 +1,36 @@
+{
+	"FileVersion": 3,
+	"Version": 1,
+	"VersionName": "0.1",
+	"FriendlyName": "DLC Support for GDK",
+	"Description": "Experimental support for DLC in GDK.",
+	"CreatedBy": "Epic Games, Inc.",
+	"Category": "Platform",
+	"CreatedByURL": "https://epicgames.com",
+	"DocsURL": "",
+	"MarketplaceURL": "",
+	"SupportURL": "",
+	"EnabledByDefault": false,
+	"RequiresBuildPlatform": false,
+	"CanContainContent": false,
+	"IsBetaVersion": true,
+	"Installed": false,
+	"ExplicitlyLoaded": true,
+	"HasExplicitPlatforms": true,	
+	"SupportedTargetPlatforms": [ "Win64" ],
+	"Modules": [
+		{
+			"Name": "GDKPlatformDLC",
+			"Type": "RuntimeNoCommandlet",
+			"LoadingPhase": "PostConfigInit",
+			"HasExplicitPlatforms": true,
+			"PlatformAllowList": [ "Win64" ]
+		}
+	],
+	"Plugins": [
+		{
+			"Name": "PlatformDLC",
+			"Enabled": true
+		}
+	]
+}
diff --git a/Engine/Plugins/Experimental/Microsoft/GDKPlatformDLC/Source/GDKPlatformDLC/GDKPlatformDLC.Build.cs b/Engine/Plugins/Experimental/Microsoft/GDKPlatformDLC/Source/GDKPlatformDLC/GDKPlatformDLC.Build.cs
new file mode 100644
--- /dev/null
+++ b/Engine/Plugins/Experimental/Microsoft/GDKPlatformDLC/Source/GDKPlatformDLC/GDKPlatformDLC.Build.cs
@@ -0,0 +1,24 @@
+// Copyright Epic Games, Inc. All Rights Reserved.
+
+using UnrealBuildTool;
+using EpicGames.Core;
+
+public class GDKPlatformDLC : ModuleRules
+{
+	public GDKPlatformDLC(ReadOnlyTargetRules Target) : base(Target)
+	{
+		PrivateDependencyModuleNames.AddAll( 
+			"GRDK",
+			"GDKRuntime",
+            "GDKPackageManifest",	
+			"PlatformDLC",
+			"Core",
+			"CoreUObject",
+			"Engine",
+			"Projects",
+			"AssetRegistry"
+		);
+
+		PublicIncludePathModuleNames.Add("GDKRuntime");
+	}
+}
diff --git a/Engine/Plugins/Experimental/Microsoft/GDKPlatformDLC/Source/GDKPlatformDLC/Private/GDKPlatformDLC.cpp b/Engine/Plugins/Experimental/Microsoft/GDKPlatformDLC/Source/GDKPlatformDLC/Private/GDKPlatformDLC.cpp
new file mode 100644
--- /dev/null
+++ b/Engine/Plugins/Experimental/Microsoft/GDKPlatformDLC/Source/GDKPlatformDLC/Private/GDKPlatformDLC.cpp
@@ -0,0 +1,1098 @@
+// Copyright Epic Games, Inc. All Rights Reserved.
+
+#include "GDKPlatformDLC.h"
+#if WITH_GRDK
+#include "GDKTaskQueueHelpers.h"
+#include "GDKThreadCheck.h"
+#include "GDKRuntimeModule.h"
+#include "IGDKPackageManifestModule.h"
+#include "HAL/PlatformFile.h"
+#include "HAL/PlatformFileManager.h"
+#include "HAL/IConsoleManager.h"
+#include "Misc/CoreDelegates.h"
+#include "Misc/ConfigCacheIni.h"
+#include "Misc/Paths.h"
+#include "Misc/ScopeLock.h"
+#include "Misc/CString.h"
+#include "Async/Async.h"
+#include "Interfaces/IPluginManager.h"
+#include "AssetRegistry/IAssetRegistry.h"
+#include "AssetRegistry/AssetRegistryState.h"
+#include "PlatformDLCModule.h"
+
+
+FGDKPlatformDLC::FGDKPlatformDLC()
+	: bIsInitialized(false)
+	, CurrentStoreContext(nullptr)
+{
+}
+
+
+FGDKPlatformDLC::~FGDKPlatformDLC()
+{
+	if (bIsInitialized)
+	{
+		Shutdown();
+	}
+}
+
+void FGDKPlatformDLC::Initialize()
+{
+	GDK_SCOPE_NOT_TIME_SENSITIVE(); // (startup only) XPackageEnumeratePackages is not safe to call on a time-sensitive thread
+
+	// read configuration
+	GConfig->GetBool(ConfigSectionName, TEXT("AutoMountDLC"), bAutoMountDLC, GEngineIni);
+	GConfig->GetBool(ConfigSectionName, TEXT("AutoMountExistingDLC"), bAutoMountExistingDLC, GEngineIni);
+	GConfig->GetBool(ConfigSectionName, TEXT("AutoMountPakFiles"), bAutoMountPakFiles, GEngineIni);
+	GConfig->GetBool(ConfigSectionName, TEXT("AutoMountPlugin"), bAutoMountPlugin, GEngineIni);
+	GConfig->GetBool(ConfigSectionName, TEXT("AutoMountAgeRestricted"), bAutoMountAgeRestricted, GEngineIni);
+	GConfig->GetBool(ConfigSectionName, TEXT("AutoAcquireLicense"), bAutoAcquireLicense, GEngineIni);
+#if !UE_BUILD_SHIPPING
+	GConfig->GetBool(ConfigSectionName, TEXT("AllowDevOnlySideLoading"), bAllowDevOnlySideLoading, GEngineIni);
+#endif
+
+	// make sure the package manifest module is loaded
+	IGDKPackageManifestModule::Get();
+
+	// use the default user for acquiring DLC licenses from the store
+	// @todo: should track user signing in & out in case this is advanced user model
+	SetStoreUser( FPlatformUserId::CreateFromInternalId(0) );
+
+	// build the StoreId map
+	TArray<FString> DLCPlugins;
+	if (GConfig->GetArray(IGDKRuntimeModule::Get().GetConfigSectionName(), TEXT("DLCPlugins"), DLCPlugins, GEngineIni))
+	{
+		FScopeLock Lock(&PackageDataLock);
+		for (FString& DLCPlugin : DLCPlugins)
+		{
+			DLCPlugin.TrimStartAndEndInline();
+			DLCPlugin.ReplaceInline(TEXT("("), TEXT(""));
+			DLCPlugin.ReplaceInline(TEXT(")"), TEXT(""));
+		
+			FString DLCName;
+			FString StoreId;
+			if (FParse::Value(*DLCPlugin, TEXT("DLCName="), DLCName) && FParse::Value(*DLCPlugin, TEXT("StoreId="), StoreId))
+			{
+				DLCNameToStoreIdMap.Add(FName(*DLCName), StoreId);
+			}
+		}
+	}
+
+	// create package data for any DLC
+	// @todo: make this async
+	auto EnumeratePackageCallback = [](void* Context, const XPackageDetails* Details)
+	{
+		FGDKPlatformDLC* This = (FGDKPlatformDLC*)Context;
+		This->AsyncInitDLC(*Details, This->bAutoMountExistingDLC);
+		return true;
+	};
+	XPackageEnumeratePackages( XPackageKind::Content, XPackageEnumerationScope::ThisOnly, this, EnumeratePackageCallback );
+
+	// register a callback when new packages are downloaded
+	auto PackageInstalledCallback = [](void* Context, const XPackageDetails* Details)
+	{
+		FGDKPlatformDLC* This = (FGDKPlatformDLC*)Context;
+		This->AsyncInitDLC(*Details, This->bAutoMountDLC);
+	};
+	XPackageRegisterPackageInstalled(FGDKAsyncTaskQueue::GetGenericQueue(), this, PackageInstalledCallback, &PackageInstalledToken);
+
+	bIsInitialized = true;
+}
+
+void FGDKPlatformDLC::Shutdown()
+{
+	// try to stop any current processing task
+	ProcessingState = EProcessingState::Destroying;
+	CurrentProcessingAsyncTask.TryCancel(true);
+	PendingPackages.Reset();
+	LostLicensePackages.Reset();
+	CurrentProcessingPackage.Reset();
+
+	// release all packages (only the package captured by AsyncGDKTask blocks will be kept alive)
+	AllPackages.Reset();
+
+
+	XPackageUnregisterPackageInstalled(PackageInstalledToken, true);
+
+	if (CurrentStoreContext != nullptr)
+	{
+		XStoreCloseContextHandle(CurrentStoreContext);
+		CurrentStoreContext = nullptr;
+	}
+
+	RemoveFromTicker();
+
+	bIsInitialized = false;
+}
+
+
+
+bool FGDKPlatformDLC::Tick( float DeltaSeconds )
+{
+	switch (ProcessingState)
+	{
+		case EProcessingState::Idle:
+		{
+			FScopeLock Lock(&PackageDataLock);
+
+			check(!CurrentProcessingPackage.IsValid());
+			if (PendingPackages.Num() > 0)
+			{
+				// get the next package to mount
+				CurrentProcessingPackage = PendingPackages.Pop();
+				InitializePackage_GameThread(CurrentProcessingPackage);
+			}
+			else if (LostLicensePackages.Num() > 0)
+			{
+				// get the next lost license package
+				CurrentProcessingPackage = LostLicensePackages.Pop();
+				LicenseLost_GameThread(CurrentProcessingPackage);
+				CurrentProcessingPackage.Reset();
+			}
+			else
+			{
+				// there's no more work to do
+				RemoveFromTicker();
+				return false;
+			}
+		}
+		break;
+
+		case EProcessingState::QueryingLicense:
+		case EProcessingState::AcquiringLicense:
+		case EProcessingState::MountingPackage:
+		case EProcessingState::CachingData:
+		{
+			// nothing to do, but sanity check that the package is still valid
+			check(CurrentProcessingPackage.IsValid());
+		}
+		break;
+
+		case EProcessingState::MountingPlugin:
+		{
+			check(CurrentProcessingPackage.IsValid());
+			MountPlugin_GameThread(CurrentProcessingPackage);
+			check(ProcessingState == EProcessingState::Completed);
+		}
+		//... fall through to completion ...
+
+		case EProcessingState::Completed:
+		{
+			check(CurrentProcessingPackage.IsValid());
+			Complete_GameThread(CurrentProcessingPackage);
+			check(!CurrentProcessingPackage.IsValid());
+			check(ProcessingState == EProcessingState::Idle);
+		}
+		break;
+	}
+
+	// keep ticking
+	return true;
+}
+
+
+void FGDKPlatformDLC::AsyncInitDLC(const XPackageDetails& Details, bool bMount)
+{
+	UE_CLOG(Details.installing, LogPlatformDLC, Error, TEXT("package %s is still installing - intelligent delivery is not supported for DLC at the moment"), UTF8_TO_TCHAR(Details.packageIdentifier) );
+
+	// add package and prepare for enumeration callbacks
+	TSharedPtr<FPackageData> Package;
+	{
+		FScopeLock Lock(&PackageDataLock);
+		
+		FString StoreId = UTF8_TO_TCHAR(Details.storeId);
+		FName DLCName;
+		for (const TTuple<FName,FString>& Pair : DLCNameToStoreIdMap)
+		{
+			if (Pair.Value == StoreId)
+			{
+				DLCName = Pair.Key;
+				break;
+			}
+		}
+		if (DLCName.IsNone())
+		{
+			UE_LOG(LogPlatformDLC, Warning, TEXT("Ignoring DLC with StoreId %s because it isn't referenced in DLCPlugins config array"), *StoreId);
+			return;
+		}
+
+		Package = AllPackages.Add_GetRef( MakeShared<FPackageData>(this, Details, DLCName));
+		Package->Context.Package = Package; // initialize the callback context. necessary to support those GDK callback functions that only take a raw context pointer
+	}
+
+	// mount the DLC package automatically, if requested (this is typically done for DLC that arrives after engine startup)
+	if (bMount)
+	{
+		AddToPendingPackages( Package );
+	}
+}
+
+
+void FGDKPlatformDLC::AddToPendingPackages( TSharedPtr<FPackageData> Package )
+{
+	FScopeLock Lock(&PackageDataLock);
+	PendingPackages.AddUnique(Package);
+	AddToTicker();
+}
+
+
+void FGDKPlatformDLC::InitializePackage_GameThread( TSharedPtr<FPackageData> Package )
+{
+	// if the package is already mounted and ready, just fire the callbacks again
+	bool bHasLicense = (Package->LicenseState == ELicenseState::Licensed || Package->LicenseState == ELicenseState::Unlicensed_DevIgnore);
+	if (Package->MountState == EMountState::Mounted && bHasLicense)
+	{
+		ProcessingState = EProcessingState::Completed;
+		return;
+	}
+
+	// check for valid state
+	if (bHasLicense && Package->MountState != EMountState::NotMounted && Package->MountState != EMountState::Unmounted)
+	{
+		checkNoEntry();
+		return;
+	}
+
+	// early out if we are ignoring age-restricted DLC
+	if (Package->bIsAgeRestricted && !bAutoMountAgeRestricted)
+	{
+		Package->MountState = EMountState::Ignored_AgeRestricted;
+		ProcessingState = EProcessingState::Completed;
+		UE_LOG(LogPlatformDLC, Log, TEXT("Ignoring DLC %s because it is age restricted for one or more of the currently logged in users"), *Package->PackageName );
+		return;
+	}
+
+	// see if we should start acquiring a license
+	bool bWantLicense = false;
+	if (bAutoAcquireLicense)
+	{
+		// make sure there is a store context
+		if (!EnsureValidStoreContext())
+		{
+			ProcessingState = EProcessingState::Completed;
+			UE_LOG(LogPlatformDLC, Error, TEXT("cannot mount DLC"));
+			return;
+		}
+
+		// see if we should acquire a license for this DLC
+		bWantLicense = (CurrentStoreContext != nullptr && Package->DLCLicense == nullptr);
+	}
+
+	if (bWantLicense)
+	{
+		ProcessingState = EProcessingState::QueryingLicense;
+		AsyncQueryLicense_GameThread(Package);
+	}
+	else
+	{
+		ProcessingState = EProcessingState::MountingPackage;
+		AsyncMountPackage_GameThread(Package);
+	}
+}
+
+
+void FGDKPlatformDLC::AsyncQueryLicense_GameThread( TSharedPtr<FPackageData> Package )
+{
+	check(Package->DLCLicense == nullptr);
+	check(ProcessingState == EProcessingState::QueryingLicense);
+
+	HRESULT hResult = AsyncGDKTask(
+		CurrentProcessingAsyncTask,
+		[Package, this](XAsyncBlock* Block)
+		{
+			return XStoreCanAcquireLicenseForPackageAsync(CurrentStoreContext, Package->PackageIdentifier, Block);
+		},
+		[Package, this](XAsyncBlock* Block)
+		{
+			// early out if the DLC system is being destroyed
+			bool bCancel = (ProcessingState == EProcessingState::Destroying);
+			if (bCancel)
+			{
+				return;
+			}
+
+			XStoreCanAcquireLicenseResult Result;
+			HRESULT hResult = XStoreCanAcquireLicenseForPackageResult(Block, &Result);
+			UE_CLOG(FAILED(hResult), LogPlatformDLC, Error, TEXT("failed to query license for package %s : 0x%X"), *Package->PackageName, hResult);
+
+			if (SUCCEEDED(hResult) && Result.status != XStoreCanLicenseStatus::NotLicensableToUser)
+			{
+				ProcessingState = EProcessingState::AcquiringLicense;
+				AsyncAcquireLicense_GameThread(Package);
+			}
+			else
+			{
+				Package->LicenseState = ELicenseState::Unlicensed;
+				ProcessingState = EProcessingState::Completed;
+				UE_CLOG(SUCCEEDED(hResult), LogPlatformDLC, Error, TEXT("Package %s is not licensable for the current store user (result = %d)"), *Package->PackageName, Result.status);
+			}
+		}
+	);
+
+	if (FAILED(hResult))
+	{
+		ProcessingState = EProcessingState::Completed;
+		UE_LOG(LogPlatformDLC, Error, TEXT("Failed to start querying licensability for package %s : 0x%X"), *Package->PackageName, hResult);
+	}
+}
+
+
+void FGDKPlatformDLC::AsyncAcquireLicense_GameThread( TSharedPtr<FPackageData> Package )
+{
+	check(Package->DLCLicense == nullptr);
+	check(ProcessingState == EProcessingState::AcquiringLicense);
+
+	HRESULT hResult = AsyncGDKTask(
+		CurrentProcessingAsyncTask,
+		[Package, this](XAsyncBlock* Block)
+		{
+			return XStoreAcquireLicenseForPackageAsync(CurrentStoreContext, Package->PackageIdentifier, Block);
+		},
+		[Package, this](XAsyncBlock* Block)
+		{
+			HRESULT hResult = XStoreAcquireLicenseForPackageResult(Block, &Package->DLCLicense);
+
+			// early out if the DLC system is being destroyed (must be done after caching the DLC license so we don't leak it)
+			bool bCancel = (ProcessingState == EProcessingState::Destroying);
+			if (bCancel)
+			{
+				return;
+			}
+
+			if (SUCCEEDED(hResult) && XStoreIsLicenseValid(Package->DLCLicense) )
+			{
+				Package->LicenseState = ELicenseState::Licensed;
+
+				// register for a callback when the license is lost
+				hResult = XStoreRegisterPackageLicenseLost(Package->DLCLicense, FGDKAsyncTaskQueue::GetGenericQueue(), &Package->Context, &FGDKPlatformDLC::OnPackageLicenseLost, &Package->PackageLicenseLostToken );
+				UE_CLOG(FAILED(hResult), LogPlatformDLC, Warning, TEXT("Could not register for GDK license lost callback : 0x%X"), hResult);
+
+				// start mounting the DLC package
+				ProcessingState = EProcessingState::MountingPackage;
+				AsyncMountPackage_GameThread(Package);
+				return;
+			}
+
+			// if we received a license that was not valid, close it now as there's nothing we can do with it
+			if (Package->DLCLicense != nullptr)
+			{
+				XStoreCloseLicenseHandle(Package->DLCLicense);
+				Package->DLCLicense = nullptr;
+			}
+
+			// handle failures
+#if !UE_BUILD_SHIPPING
+			if (bAllowDevOnlySideLoading && (
+				hResult == 0x87e10bc6 || // LM_E_CONTENT_NOT_IN_CATALOG
+				hResult == 0x803f9006 || // LM_E_ENTITLED_USER_SIGNED_OUT
+				hResult == 0x87dd0013))  // AM_E_USER_NOT_FOUND
+			{
+				// this is most likely from using the GDK development tools to sideload the DLC package instead of aquiring it from the Store. Allow the mount to happen regardless
+				UE_LOG(LogPlatformDLC, Warning, TEXT("Ignoring DLC license aquisition failure (0x%X) in non-shipping build. This will fail in Shipping"), hResult );
+				Package->LicenseState = ELicenseState::Unlicensed_DevIgnore;
+
+				ProcessingState = EProcessingState::MountingPackage;
+				AsyncMountPackage_GameThread(Package);
+			}
+			else
+#endif
+			{
+				Package->LicenseState = ELicenseState::Unlicensed;
+				ProcessingState = EProcessingState::Completed;
+				UE_LOG(LogPlatformDLC, Error, TEXT("Failed to get license for package %s : 0x%X"), *Package->PackageName, hResult);
+			}
+		}
+	);
+
+	if (FAILED(hResult))
+	{
+		Package->LicenseState = ELicenseState::Unlicensed;
+		ProcessingState = EProcessingState::Completed;
+		UE_LOG(LogPlatformDLC, Error, TEXT("Failed to start getting license for package %s : 0x%X"), *Package->PackageName, hResult);
+	}
+}
+
+
+void FGDKPlatformDLC::AsyncMountPackage_GameThread( TSharedPtr<FPackageData> Package )
+{
+	check(ProcessingState == EProcessingState::MountingPackage);
+
+	// if the package is already mounted, it means we've only reaquired the license
+	if (Package->DLCHandle != nullptr && ensure(Package->MountState == EMountState::Mounted) )
+	{
+		ProcessingState = EProcessingState::Completed;
+		PostNotification(Package->DLCName, ENotification::Entitlement, true);
+		return;
+	}
+
+
+	HRESULT hResult = AsyncGDKTask(
+		CurrentProcessingAsyncTask,
+		[Package](XAsyncBlock* Block)
+		{
+			return XPackageMountWithUiAsync(Package->PackageIdentifier, Block);
+		},
+		[Package, this](XAsyncBlock* Block)
+		{
+			HRESULT hResult = XPackageMountWithUiResult(Block, &Package->DLCHandle);
+
+			// early out if the DLC system is being destroyed (must be done after caching the DLC handle so we don't leak it)
+			bool bCancel = (ProcessingState == EProcessingState::Destroying);
+			if (bCancel)
+			{
+				return;
+			}
+
+			if (SUCCEEDED(hResult))
+			{
+				// continue to initialize the DLC package on a background thread
+				ProcessingState = EProcessingState::CachingData;
+				AsyncTask( ENamedThreads::AnyBackgroundThreadNormalTask, [Package, this]()
+				{
+					CachePackageData_BackgroundThread(Package);
+				});
+				return;
+			}
+
+			Package->MountState = EMountState::NotMounted;
+			ProcessingState = EProcessingState::Completed;
+			UE_LOG(LogPlatformDLC, Error, TEXT("Failed to mount DLC '%s' : 0x%X"), *Package->PackageName, hResult);
+		}
+	);
+
+	if (FAILED(hResult))
+	{
+		Package->MountState = EMountState::NotMounted;
+		ProcessingState = EProcessingState::Completed;
+		UE_LOG(LogPlatformDLC, Error, TEXT("Failed to start mounting mount DLC '%s' : 0x%X"), *Package->PackageName, hResult);
+	}
+}
+
+
+void FGDKPlatformDLC::CachePackageData_BackgroundThread( TSharedPtr<FPackageData> Package )
+{
+	check(ProcessingState == EProcessingState::CachingData);
+
+	HRESULT hResult;
+
+	// get the mount path
+	char MountPath[MAX_PATH+1];
+	hResult = XPackageGetMountPath(Package->DLCHandle, UE_ARRAY_COUNT(MountPath)-1, MountPath);
+	if (FAILED(hResult))
+	{
+		Package->MountState = EMountState::NotMounted;
+		ProcessingState = EProcessingState::Completed;
+		UE_LOG(LogPlatformDLC, Error, TEXT("Failed to get mount path for DLC '%s' : 0x%X"), *Package->PackageName, hResult);
+
+		XPackageCloseMountHandle(Package->DLCHandle);
+		Package->DLCHandle = nullptr;
+		return;
+	}
+	Package->MountPath = UTF8_TO_TCHAR(MountPath);
+
+
+	// load package manifest
+	IGDKPackageManifestModule::Get().LoadManifestFromDLC(Package->DLCHandle, Package->PackageIdentifier, Package->MountPath);
+
+	// finalize on the GameThread
+	ProcessingState = EProcessingState::MountingPlugin;
+}
+
+
+void FGDKPlatformDLC::MountPlugin_GameThread( TSharedPtr<FPackageData> Package )
+{
+	check(ProcessingState == EProcessingState::MountingPlugin);
+
+	if (bAutoMountPakFiles)
+	{
+		// mount all available pak files. NOTE: intelligent delivery is not supported for DLC at the moment - only what is currently installed will be mounted
+		TArray<FString> PakFileNames;
+		IPlatformFile::GetPlatformPhysical().FindFilesRecursively(PakFileNames, *Package->MountPath, TEXT(".pak") );
+		for (const FString& PakFileName : PakFileNames)
+		{
+			if (IGDKPackageManifestModule::Get().IsPakFileInstalled(Package->DLCHandle, PakFileName))
+			{
+				IPakFile* PakFile = FCoreDelegates::MountPak.IsBound() ? FCoreDelegates::MountPak.Execute(PakFileName, INDEX_NONE) : nullptr;
+				if (PakFile)
+				{
+					UE_LOG(LogPlatformDLC, Log, TEXT("Successfully mounted %s"), *PakFileName);
+				}
+				else
+				{
+					UE_LOG(LogPlatformDLC, Log, TEXT("Failed to mount %s"), *PakFileName);
+				}
+			}
+		}
+	}
+
+	if (bAutoMountPlugin && bAutoMountPakFiles)
+	{
+		// find & load the .uplugin (should only be one per DLC, as specified by "RunUAT ... -DLCName=<plugin>")
+		TArray<FString> FoundPlugins;
+		IPlatformFile::GetPlatformPhysical().FindFilesRecursively(FoundPlugins, *Package->MountPath, TEXT(".uplugin") );
+		checkf(FoundPlugins.Num() != 0, TEXT("found %d plugins in %s - DLC should have 1!"), FoundPlugins.Num(), *Package->MountPath);
+		const FString PluginFileName = FoundPlugins.Last();
+
+		IPluginManager::Get().AddToPluginsList(PluginFileName);
+		FString DLCName = FPaths::GetBaseFilename(PluginFileName);
+
+		Package->DLCName = FName(DLCName);
+		Package->DLCPlugin = IPluginManager::Get().FindPlugin(DLCName);
+		checkf(Package->DLCPlugin.IsValid(), TEXT("failed to find DLC plugin %s"), *Package->DLCName.ToString());
+
+
+		// ensure the plugin is mounted - will also load the shader code library if there is one
+		IPluginManager::Get().MountNewlyCreatedPlugin(Package->DLCName.ToString());
+
+		// load the asset registry for the new plugin
+		if (Package->DLCPlugin->CanContainContent())
+		{
+			FString AssetRegistryFileName = Package->DLCPlugin->GetBaseDir() / TEXT("AssetRegistry.bin");
+			AssetRegistryFileName.ReplaceInline(*Package->MountPath, TEXT("../../..") );
+
+			FAssetRegistryState PluginAssetRegistry;
+			if (FAssetRegistryState::LoadFromDisk(*AssetRegistryFileName, FAssetRegistryLoadOptions(), PluginAssetRegistry))
+			{
+				IAssetRegistry::GetChecked().AppendState(PluginAssetRegistry);
+			}
+			else
+			{
+				UE_LOG(LogPlatformDLC, Error, TEXT("Failed to load plugin asset registry state %s"), *AssetRegistryFileName);
+			}
+		}
+	}
+
+	// all done
+	Package->MountState = EMountState::Mounted;
+	ProcessingState = EProcessingState::Completed;
+}
+
+
+void FGDKPlatformDLC::Complete_GameThread( TSharedPtr<FPackageData> Package )
+{
+	check(IsInGameThread());
+
+	// signal callback
+	bool bSuccess = (Package->MountState == EMountState::Mounted);
+	PostNotification(Package->DLCName, ENotification::Mounted, bSuccess);
+
+	// remove the current work item & check for more work next tick
+	CurrentProcessingPackage.Reset();
+	ProcessingState = EProcessingState::Idle;
+}
+
+
+void FGDKPlatformDLC::LicenseLost_GameThread( TSharedPtr<FPackageData> Package )
+{
+	// close the existing license
+	if (Package->DLCLicense != nullptr)
+	{
+		XStoreUnregisterPackageLicenseLost( Package->DLCLicense, Package->PackageLicenseLostToken, true);
+		XStoreCloseLicenseHandle(Package->DLCLicense);
+		Package->DLCLicense = nullptr;
+	}
+	Package->LicenseState = ELicenseState::LicenseLost;
+
+	// signal that the license has been lost. it's up to the title whether they want to retry/unmount/terminate etc.
+	PostNotification(Package->DLCName, ENotification::Entitlement, false);
+}
+
+
+
+bool FGDKPlatformDLC::Unmount( FName DLCName )
+{
+	TSharedPtr<FGDKPlatformDLC::FPackageData> Package = FindDLC(DLCName);
+	if (!Package.IsValid())
+	{
+		UE_LOG(LogPlatformDLC, Error, TEXT("DLC %s not known"), *DLCName.ToString());
+		return false;	
+	}
+	
+	bool bSuccess = false;
+	// @todo: make this async
+
+	FScopeLock Lock(&PackageDataLock);
+
+	if (Package->DLCHandle != nullptr)
+	{
+		// attempt to unregister the plugin as much as possible
+		if (bAutoMountPlugin && bAutoMountPakFiles)
+		{
+			FText FailureReason;
+			if (Package->DLCPlugin->GetDescriptor().bExplicitlyLoaded && IPluginManager::Get().UnmountExplicitlyLoadedPlugin(Package->DLCName.ToString(), &FailureReason))
+			{
+				const FString PluginFileName = FPaths::Combine( Package->DLCPlugin->GetBaseDir(), Package->DLCName.ToString() + TEXT(".uplugin") );
+				IPluginManager::Get().RemoveFromPluginsList( PluginFileName );
+			}
+			else
+			{
+				// not possible to properly unmount, so just simulate the results
+				IPluginManager::Get().OnPluginUnmounted().Broadcast(*Package->DLCPlugin);
+				FPackageName::UnRegisterMountPoint( Package->DLCPlugin->GetMountedAssetPath(), Package->DLCPlugin->GetContentDir() );
+			}
+		}
+
+		// unmount all pak files
+		if (bAutoMountPakFiles)
+		{
+			TArray<FString> PakFileNames;
+			IPlatformFile::GetPlatformPhysical().FindFilesRecursively(PakFileNames, *Package->MountPath, TEXT(".pak") );
+			for (const FString& PakFileName : PakFileNames)
+			{
+				if (IGDKPackageManifestModule::Get().IsPakFileInstalled(Package->DLCHandle, PakFileName))
+				{
+					bool bUnmounted = FCoreDelegates::OnUnmountPak.IsBound() ? FCoreDelegates::OnUnmountPak.Execute(PakFileName) : false;
+					if (bUnmounted)
+					{
+						UE_LOG(LogPlatformDLC, Log, TEXT("Successfully unmounted %s"), *PakFileName);
+					}
+					else
+					{
+						UE_LOG(LogPlatformDLC, Log, TEXT("Failed to unmount %s"), *PakFileName);
+					}
+				}
+			}
+		}
+
+		// unmount the package
+		XPackageCloseMountHandle(Package->DLCHandle);
+		Package->DLCHandle = nullptr;
+
+		Package->MountState = EMountState::Unmounted;
+
+		bSuccess = true;
+	}
+
+	PostNotification(DLCName, ENotification::Unmounted, bSuccess);
+	return true;
+}
+
+
+void FGDKPlatformDLC::AddToTicker()
+{
+	if (!TickHandle.IsValid())
+	{
+		FTSTicker& Ticker = FTSTicker::GetCoreTicker();
+
+		// Register delegate for ticker callback
+		FTickerDelegate TickDelegate = FTickerDelegate::CreateRaw(this, &FGDKPlatformDLC::Tick);
+		TickHandle = Ticker.AddTicker(TickDelegate, 0.1f);
+	}
+}
+
+void FGDKPlatformDLC::RemoveFromTicker()
+{
+	FTSTicker& Ticker = FTSTicker::GetCoreTicker();
+
+	// Unregister ticker delegate
+	if (TickHandle.IsValid())
+	{
+		Ticker.RemoveTicker(TickHandle);
+		TickHandle.Reset();
+	}
+}
+
+
+
+void FGDKPlatformDLC::OnPackageLicenseLost(void* Context)
+{
+	TSharedPtr<FPackageData> Package = ((FPackageData::FContext*)Context)->Package.Pin();
+	{
+		FScopeLock Lock(&Package->Owner->PackageDataLock);
+		Package->Owner->LostLicensePackages.AddUnique(Package);
+		Package->Owner->AddToTicker();
+
+		XStoreUnregisterPackageLicenseLost( Package->DLCLicense, Package->PackageLicenseLostToken, false);
+	}
+}
+
+
+
+void FGDKPlatformDLC::SetStoreUser( FPlatformUserId PlatformUserId )
+{
+	FGDKUserHandle NewStoreUser = IGDKRuntimeModule::Get().GetUserHandleByPlatformId(PlatformUserId);
+	if (NewStoreUser != CurrentStoreUser)
+	{
+		CurrentStoreUser = NewStoreUser;
+		CurrentStoreUserId = PlatformUserId;
+
+		if (CurrentStoreContext != nullptr)
+		{
+			XStoreCloseContextHandle(CurrentStoreContext);
+			CurrentStoreContext = nullptr;
+		}
+	}
+}
+
+bool FGDKPlatformDLC::EnsureValidStoreContext()
+{
+	if (CurrentStoreUser.IsValid() && CurrentStoreContext == nullptr)
+	{
+		GDK_SCOPE_NOT_TIME_SENSITIVE(); // XStoreCreateContext is not safe to call on a time-sensitive thread
+
+		HRESULT hResult = XStoreCreateContext(CurrentStoreUser, &CurrentStoreContext);
+		if (FAILED(hResult))
+		{
+			UE_LOG(LogPlatformDLC, Error, TEXT("Failed to create store context for user %s : 0x%X"), *LexToString(CurrentStoreUser), hResult);
+			return false;
+		}
+	}
+
+	return true;
+}
+
+
+
+TSharedPtr<FGDKPlatformDLC::FPackageData> FGDKPlatformDLC::FindDLC( FName DLCName )
+{
+	FScopeLock Lock(&PackageDataLock);
+
+	for (TSharedPtr<FPackageData> PackageData : AllPackages)
+	{
+		if (PackageData->DLCName == DLCName)
+		{
+			return PackageData;
+		}
+	}
+
+	return nullptr;
+}
+
+TSharedPtr<const FGDKPlatformDLC::FPackageData> FGDKPlatformDLC::FindDLC( FName DLCName ) const
+{
+	TSharedPtr<FPackageData> Result = const_cast<FGDKPlatformDLC*>(this)->FindDLC(DLCName);
+	return Result;
+}
+
+bool FGDKPlatformDLC::IsInitialized() const
+{
+	return bIsInitialized;
+}
+
+bool FGDKPlatformDLC::HasEntitlement( FName DLCName ) const
+{
+	TSharedPtr<const FPackageData> Package = FindDLC(DLCName);
+	return Package.IsValid() && (Package->LicenseState == ELicenseState::Licensed || Package->LicenseState == ELicenseState::Unlicensed_DevIgnore);
+}
+
+bool FGDKPlatformDLC::IsMounted( FName DLCName ) const
+{
+	TSharedPtr<const FPackageData> Package = FindDLC(DLCName);
+	return Package.IsValid() && (Package->MountState == EMountState::Mounted);
+}
+
+bool FGDKPlatformDLC::IsDownloaded( FName DLCName ) const
+{
+	// FPackageData only exists if it's downloaded
+	TSharedPtr<const FPackageData> Package = FindDLC(DLCName);
+	return (Package.IsValid());
+}
+
+
+bool FGDKPlatformDLC::Mount( FName DLCName )
+{
+	TSharedPtr<FPackageData> Package = FindDLC(DLCName);
+	if (!Package.IsValid())
+	{
+		return false;
+	}
+
+	AddToPendingPackages(Package);
+
+	return true;
+}
+
+
+FString FGDKPlatformDLC::GetRootDirectory( FName DLCName ) const
+{
+	TSharedPtr<const FPackageData> Package = FindDLC(DLCName);
+	if (ensure(Package.IsValid() && Package->MountState == EMountState::Mounted))
+	{
+		return Package->MountPath;
+	}
+
+	return FString();
+}
+
+TArray<FName> FGDKPlatformDLC::GetAllDLCNames() const
+{
+	FScopeLock Lock(&PackageDataLock);
+
+	TArray<FName> DLCNames;
+	DLCNameToStoreIdMap.GetKeys(DLCNames);
+
+	return MoveTemp(DLCNames);
+}
+
+TArray<FName> FGDKPlatformDLC::GetMountedDLCNames() const
+{
+	TArray<FName> DLCNames;
+	for (const FName& DLCName : GetAllDLCNames())
+	{
+		if (IsMounted(DLCName) && HasEntitlement(DLCName))
+		{
+			DLCNames.Add(DLCName);
+		}
+	}
+
+	return MoveTemp(DLCNames);
+}
+
+FString FGDKPlatformDLC::GetStoreId( FName DLCName ) const
+{
+	const FString* StoreIdPtr = DLCNameToStoreIdMap.Find(DLCName);
+	if (StoreIdPtr != nullptr)
+	{
+		return *StoreIdPtr;
+	}
+
+	return FString();
+}
+
+FGDKPlatformDLC::FOnDLCNotification& FGDKPlatformDLC::OnNotification()
+{
+	return OnNotificationDelegate;
+}
+
+void FGDKPlatformDLC::PostNotification(FName DLCName, ENotification Notification, bool bSuccess)
+{
+	ExecuteOnGameThread( UE_SOURCE_LOCATION, [this, DLCName, Notification, bSuccess]
+	{
+		OnNotificationDelegate.Broadcast(DLCName, Notification, bSuccess);
+	});
+}
+
+
+
+
+FGDKPlatformDLC::FPackageData::FPackageData(FGDKPlatformDLC* InOwner, const XPackageDetails& InDetails, const FName& InDLCName)
+	: MountState(EMountState::NotMounted)
+	, LicenseState(ELicenseState::Unlicensed)
+	, Owner(InOwner)
+	, bIsAgeRestricted(InDetails.ageRestricted)
+	, DLCName(InDLCName)
+	, DLCHandle(nullptr)
+	, DLCLicense(nullptr)
+{
+	FCStringAnsi::Strncpy(PackageIdentifier, InDetails.packageIdentifier, XPACKAGE_IDENTIFIER_MAX_LENGTH);
+	FCStringAnsi::Strncpy(StoreId, InDetails.storeId, STORE_SKU_ID_SIZE);
+
+	// set the name
+	PackageName = UTF8_TO_TCHAR(InDetails.displayName);
+	if (PackageName.IsEmpty())
+	{
+		PackageName = UTF8_TO_TCHAR(PackageIdentifier);
+	}
+}
+
+
+
+FGDKPlatformDLC::FPackageData::~FPackageData()
+{
+	if (DLCHandle != nullptr)
+	{
+		XPackageCloseMountHandle(DLCHandle);
+		DLCHandle = nullptr;
+	}
+
+	if (DLCLicense != nullptr)
+	{
+		XStoreUnregisterPackageLicenseLost( DLCLicense, PackageLicenseLostToken, true);
+		XStoreCloseLicenseHandle(DLCLicense);
+		DLCLicense = nullptr;
+	}
+}
+
+
+
+
+bool FGDKPlatformDLC::Uninstall( FName DLCName )
+{
+	GDK_SCOPE_NOT_TIME_SENSITIVE(); // (debug only) XPackageUninstallPackage is not safe to call on a time-sensitive thread
+
+	TSharedPtr<FGDKPlatformDLC::FPackageData> Package = FindDLC(DLCName);
+	if (Package.IsValid())
+	{
+		if (Package->DLCHandle == nullptr)
+		{
+			bool bSuccess = XPackageUninstallPackage(TCHAR_TO_UTF8(Package->PackageIdentifier));
+			if (bSuccess)
+			{
+				FScopeLock Lock(&PackageDataLock);
+				AllPackages.Remove(Package);
+			}
+			PostNotification(Package->DLCName, ENotification::Uninstalled, bSuccess);
+			return true;
+		}
+		else
+		{
+			UE_LOG(LogPlatformDLC, Error, TEXT("DLC %s must be unmounted before it can be uninstalled"), *DLCName.ToString());
+			return false;
+		}
+
+	}
+
+	return false;
+}
+
+bool FGDKPlatformDLC::Download( FName DLCName )
+{
+	TSharedPtr<FGDKPlatformDLC::FPackageData> Package = FindDLC(DLCName);
+	if (Package.IsValid())
+	{
+		// if we know about the package, it's already downloaded
+		PostNotification(DLCName, ENotification::Downloaded, true);
+		return true;
+	}
+	else
+	{
+		// find the package with the store id
+		FString StoreId = GetStoreId(DLCName);
+		if (!StoreId.IsEmpty())
+		{
+			// make sure there is a store context
+			if (!EnsureValidStoreContext())
+			{
+				return false;
+			}
+
+			// kick off the async 
+			HRESULT hResult = AsyncGDKTask( 
+				[this,StoreId]( XAsyncBlock* AsyncBlock )
+				{
+					char StoreIdAnsi[64];
+					FCStringAnsi::Strncpy(StoreIdAnsi, TCHAR_TO_ANSI(*StoreId), UE_ARRAY_COUNT(StoreIdAnsi));
+
+					const char* StoreIds[] =
+					{
+						StoreIdAnsi
+					};
+
+					return XStoreDownloadAndInstallPackagesAsync( CurrentStoreContext, StoreIds, 1, AsyncBlock );
+				},
+				[this, DLCName]( XAsyncBlock* AsyncBlock )
+				{
+					HRESULT hResult;
+
+					uint32 NumPackages = 0;
+					hResult = XStoreDownloadAndInstallPackagesResultCount( AsyncBlock, &NumPackages );
+					if (FAILED(hResult))
+					{
+						UE_LOG(LogPlatformDLC, Error, TEXT("failed to get installing packages count for %s : 0x%X"), *DLCName.ToString(), hResult );
+						PostNotification(DLCName, ENotification::Downloaded, false);
+						return;
+					}
+
+					auto PackageIdentifiers = new char[NumPackages][XPACKAGE_IDENTIFIER_MAX_LENGTH];
+					hResult = XStoreDownloadAndInstallPackagesResult( AsyncBlock, NumPackages, &PackageIdentifiers[0] );
+					if (FAILED(hResult))
+					{
+						delete[] PackageIdentifiers;
+						UE_LOG(LogPlatformDLC, Error, TEXT("failed to get installing packages identifiers for %s : 0x%X"), *DLCName.ToString(), hResult );
+						PostNotification(DLCName, ENotification::Downloaded, false);
+						return;
+					}
+
+					// the download was successfully started. we must have a valid license for this before we can mount it
+					for (uint32 Index = 0; Index < NumPackages; Index++)
+					{
+						UE_LOG(LogPlatformDLC, Log, TEXT("Downloading package %s for DLC %s"), UTF8_TO_TCHAR(PackageIdentifiers[Index]), *DLCName.ToString() );
+					}
+	
+					delete[] PackageIdentifiers;
+
+					PostNotification(DLCName, ENotification::Downloaded, true);
+				}
+			);
+
+			return SUCCEEDED(hResult);
+		}
+
+		return false;
+	}
+}
+
+#if !UE_BUILD_SHIPPING
+FString FGDKPlatformDLC::Debug_GetStateDescription() const
+{
+	return FString::Printf( TEXT("GDK DLC.  Store User: %-3d %-16s   Processing State: %s"), CurrentStoreUserId.GetInternalId(), *IGDKRuntimeModule::Get().GetGamertag(CurrentStoreUser), *LexToString(ProcessingState) );
+
+}
+
+bool FGDKPlatformDLC::Debug_GetPackageDescription( int PackageIndex, FString& OutDescription, FColor& OutColor ) const
+{
+	TSharedPtr<FPackageData> Package;
+	{
+		FScopeLock Lock(&PackageDataLock);
+		if (!AllPackages.IsValidIndex(PackageIndex))
+		{
+			return false;
+		}
+		Package = AllPackages[PackageIndex];
+	}
+
+	if (Package->MountState == EMountState::NotMounted || Package->MountState == EMountState::Unmounted)
+	{
+		OutColor = FColor::Silver;
+	}
+	else if (Package->MountState == EMountState::Mounted)
+	{
+		bool bIsLicensed = (Package->LicenseState == ELicenseState::Licensed) || (Package->LicenseState == ELicenseState::Unlicensed_DevIgnore);
+		OutColor = bIsLicensed ? FColor::Green : FColor::Yellow;
+	}
+	else if (Package->MountState < EMountState::Mounted)
+	{
+		OutColor = FColor::Yellow;
+	}
+	else
+	{
+		OutColor = FColor::Red;
+	}
+
+	OutDescription = FString::Printf(TEXT(".......%-30s %-30s %-30s (%s)"), *Package->DLCName.ToString(), *LexToString(Package->MountState), *LexToString(Package->LicenseState), *Package->PackageName );
+	return true;
+}
+#endif //UE_BUILD_SHIPPING
+
+
+FString LexToString( FGDKPlatformDLC::EMountState MountState )
+{
+	switch (MountState)
+	{
+		case FGDKPlatformDLC::EMountState::NotMounted:              return TEXT("NotMounted");
+		case FGDKPlatformDLC::EMountState::Mounted:                 return TEXT("Mounted");
+		case FGDKPlatformDLC::EMountState::Unmounted:               return TEXT("Unmounted");
+		case FGDKPlatformDLC::EMountState::Ignored_AgeRestricted:   return TEXT("Ignored(AgeRestricted)");
+		default:                                                    return TEXT("Invalid");
+	}
+}
+
+FString LexToString( FGDKPlatformDLC::ELicenseState LicenseState )
+{
+	switch(LicenseState)
+	{
+		case FGDKPlatformDLC::ELicenseState::Unlicensed:            return TEXT("Unlicensed");
+		case FGDKPlatformDLC::ELicenseState::Licensed:              return TEXT("Licensed");
+		case FGDKPlatformDLC::ELicenseState::LicenseLost:           return TEXT("LicenseLost");
+		case FGDKPlatformDLC::ELicenseState::Unlicensed_DevIgnore:  return TEXT("Unlicensed(Sideload)");
+        default:                                                    return TEXT("Invalid");
+	}
+}
+
+FString LexToString( FGDKPlatformDLC::EProcessingState ProcessingState )
+{
+	switch(ProcessingState)
+	{
+        case FGDKPlatformDLC::EProcessingState::Idle:               return TEXT("Idle");
+		case FGDKPlatformDLC::EProcessingState::QueryingLicense:    return TEXT("QueryingLicense");
+		case FGDKPlatformDLC::EProcessingState::AcquiringLicense:   return TEXT("AcquiringLicense");
+		case FGDKPlatformDLC::EProcessingState::MountingPackage:    return TEXT("MountingPackage");
+		case FGDKPlatformDLC::EProcessingState::CachingData:        return TEXT("CachingData");
+		case FGDKPlatformDLC::EProcessingState::MountingPlugin:     return TEXT("MountingPlugin");
+        case FGDKPlatformDLC::EProcessingState::Completed:          return TEXT("Completed");
+		case FGDKPlatformDLC::EProcessingState::Destroying:         return TEXT("Destroying");
+        default:                                                    return TEXT("Invalid");
+	}
+}
+
+#endif //WITH_GRDK
diff --git a/Engine/Plugins/Experimental/Microsoft/GDKPlatformDLC/Source/GDKPlatformDLC/Private/GDKPlatformDLC.h b/Engine/Plugins/Experimental/Microsoft/GDKPlatformDLC/Source/GDKPlatformDLC/Private/GDKPlatformDLC.h
new file mode 100644
--- /dev/null
+++ b/Engine/Plugins/Experimental/Microsoft/GDKPlatformDLC/Source/GDKPlatformDLC/Private/GDKPlatformDLC.h
@@ -0,0 +1,174 @@
+// Copyright Epic Games, Inc. All Rights Reserved.
+
+#pragma once
+
+#include "PlatformDLC.h"
+#include "GDKRuntimeModule.h"
+#if WITH_GRDK
+#include "GDKPlatformDLCModule.h"
+#include "GDKTaskQueueHelpers.h"
+#include "Containers/Array.h"
+#include "Containers/Ticker.h"
+#include "Math/Color.h"
+
+#include "Microsoft/AllowMicrosoftPlatformTypes.h"
+THIRD_PARTY_INCLUDES_START
+#include <XPackage.h>
+#include <XStore.h>
+#include <atomic>
+THIRD_PARTY_INCLUDES_END
+#include "Microsoft/HideMicrosoftPlatformTypes.h"
+
+class FGDKPlatformDLC : public IPlatformDLC
+{
+public:
+	FGDKPlatformDLC();
+	virtual ~FGDKPlatformDLC();
+
+	// IPlatformDLC implementation
+	virtual void Initialize() override;
+	virtual void Shutdown() override;
+	virtual bool IsInitialized() const override;
+
+	virtual bool HasEntitlement( FName DLCName ) const override;
+	virtual bool IsMounted( FName DLCName ) const override;
+	virtual bool IsDownloaded( FName DLCName ) const override;
+
+	virtual FOnDLCNotification& OnNotification() override;
+
+	virtual bool Mount( FName DLCName ) override;
+	virtual bool Unmount( FName DLCName ) override;
+	virtual bool Download( FName DLCName ) override;
+	virtual bool Uninstall( FName DLCName ) override;
+
+	virtual FString GetRootDirectory( FName DLCName ) const override;
+	virtual TArray<FName> GetAllDLCNames() const override;
+	virtual TArray<FName> GetMountedDLCNames() const override;
+	virtual FString GetStoreId( FName DLCName ) const override;
+
+
+	// support for debug rendering
+#if !UE_BUILD_SHIPPING
+	FString Debug_GetStateDescription() const;
+	bool Debug_GetPackageDescription( int PackageIndex, FString& OutDescription, FColor& OutColor ) const;
+#endif
+
+private:
+	FOnDLCNotification OnNotificationDelegate;
+
+	void SetStoreUser( FPlatformUserId PlatformUserId );
+	void PostNotification(FName DLCName, ENotification Notificiation, bool bSuccess);
+
+
+	enum class EMountState : uint8
+	{
+		NotMounted,
+		Mounted,
+		Unmounted,
+		Ignored_AgeRestricted,
+	};
+	enum class ELicenseState : uint8
+	{
+		Unlicensed,
+		Licensed,
+		LicenseLost,
+		Unlicensed_DevIgnore,
+	};
+
+	struct FPackageData
+	{
+		FPackageData(FGDKPlatformDLC* InOwner, const XPackageDetails& InDetails, const FName& InDLCName);
+		~FPackageData();
+
+		char                        PackageIdentifier[XPACKAGE_IDENTIFIER_MAX_LENGTH];
+		char                        StoreId[STORE_SKU_ID_SIZE];
+		std::atomic<EMountState>    MountState;
+		std::atomic<ELicenseState>  LicenseState;
+		FString                     PackageName;
+		FString                     MountPath;
+		FGDKPlatformDLC*            Owner;
+		bool                        bIsAgeRestricted;
+		FName                       DLCName;
+		XPackageMountHandle         DLCHandle;
+		XStoreLicenseHandle         DLCLicense;
+		XTaskQueueRegistrationToken PackageLicenseLostToken;
+		TSharedPtr<class IPlugin>   DLCPlugin;
+
+		// context used for GDK callbacks that only accept void*
+		struct FContext
+		{
+			TWeakPtr<FPackageData> Package;
+		};
+		FContext Context;
+	};
+	TArray<TSharedPtr<FPackageData>> AllPackages;
+	TMap<FName, FString>             DLCNameToStoreIdMap;
+	mutable FCriticalSection         PackageDataLock;
+	XTaskQueueRegistrationToken      PackageInstalledToken;
+	std::atomic<bool>                bIsInitialized;
+
+	void AsyncInitDLC(const XPackageDetails& Details, bool bMount);
+	TSharedPtr<FPackageData> FindDLC( FName DLCName );
+	TSharedPtr<const FPackageData> FindDLC( FName DLCName ) const;
+
+	// DLC processing
+	enum class EProcessingState : uint8
+	{
+		Idle,
+		QueryingLicense,
+		AcquiringLicense,
+		MountingPackage,
+		CachingData,
+		MountingPlugin,
+		Completed,
+		Destroying,
+	};
+	std::atomic<EProcessingState>    ProcessingState = EProcessingState::Idle;
+	TArray<TSharedPtr<FPackageData>> PendingPackages;
+	TArray<TSharedPtr<FPackageData>> LostLicensePackages;
+	TSharedPtr<FPackageData>         CurrentProcessingPackage;
+	FGDKAsyncTaskMonitor             CurrentProcessingAsyncTask;
+
+	void AddToPendingPackages( TSharedPtr<FPackageData> Package );
+	void InitializePackage_GameThread( TSharedPtr<FPackageData> Package );
+	void AsyncQueryLicense_GameThread( TSharedPtr<FPackageData> Package );
+	void AsyncAcquireLicense_GameThread( TSharedPtr<FPackageData> Package );
+	void AsyncMountPackage_GameThread( TSharedPtr<FPackageData> Package );
+	void CachePackageData_BackgroundThread( TSharedPtr<FPackageData> Package );
+	void MountPlugin_GameThread( TSharedPtr<FPackageData> Package );
+	void Complete_GameThread( TSharedPtr<FPackageData> Package );
+	void LicenseLost_GameThread( TSharedPtr<FPackageData> Package );
+
+	// store
+	XStoreContextHandle CurrentStoreContext;
+	FGDKUserHandle      CurrentStoreUser;
+	FPlatformUserId     CurrentStoreUserId;
+	bool EnsureValidStoreContext();
+	static void OnPackageLicenseLost(void* Context);
+
+	// configuration
+	bool bAutoMountDLC = true;            // Whether to automatically mount any new DLC that is downloaded once the application is up and running.
+	bool bAutoMountExistingDLC = true;    // Whether to automatically mount any already-downloaded DLC on system startup.
+	bool bAutoMountPakFiles = true;       // Whether to automatically mount pak files
+	bool bAutoMountPlugin = true;         // Whether to automatically mount the DLC plugin and asset registry etc. (requires bAutoMountPakFiles)
+	bool bAutoMountAgeRestricted = false; // Whether to attempt to mount age-restricted DLC. This will trigger a parental controls consent prompt.
+	bool bAutoAcquireLicense = true;      // Whether to attempt to acquire the license for any DLC before attempting to mount. If this is false, the title must manage licensing directly.
+#if !UE_BUILD_SHIPPING
+	bool bAllowDevOnlySideLoading = true; // Whether to ignore certain failure states when license aquisition fails in a non-shipping build - typically during dev sideloading scenarios.
+#endif
+
+
+
+	// ticker
+	FTSTicker::FDelegateHandle TickHandle;
+	void AddToTicker();
+	void RemoveFromTicker();
+	bool Tick( float DeltaSeconds );
+
+	// internal LexToString
+	friend FString LexToString( EMountState MountState );
+	friend FString LexToString( ELicenseState LicenseState );
+	friend FString LexToString( EProcessingState ProcessingState );
+};
+
+#endif //WITH_GRDK
diff --git a/Engine/Plugins/Experimental/Microsoft/GDKPlatformDLC/Source/GDKPlatformDLC/Private/GDKPlatformDLCModule.cpp b/Engine/Plugins/Experimental/Microsoft/GDKPlatformDLC/Source/GDKPlatformDLC/Private/GDKPlatformDLCModule.cpp
new file mode 100644
--- /dev/null
+++ b/Engine/Plugins/Experimental/Microsoft/GDKPlatformDLC/Source/GDKPlatformDLC/Private/GDKPlatformDLCModule.cpp
@@ -0,0 +1,149 @@
+// Copyright Epic Games, Inc. All Rights Reserved.
+
+#include "GDKPlatformDLCModule.h"
+#include "GDKPlatformDLC.h"
+#include "GDKRuntimeModule.h"
+#include "Modules/ModuleInterface.h"
+#include "Modules/ModuleManager.h"
+#if WITH_GRDK
+#include "PlatformDLCModule.h"
+#include "Templates/SharedPointer.h"
+#include "HAL/IConsoleManager.h"
+#include "Engine/Engine.h"
+#include "Containers/Ticker.h"
+#endif
+
+class FGDKPlatformDLCModule : public IGDKPlatformDLCModule
+{
+public:
+
+#if WITH_GRDK
+	virtual void StartupModule() override
+	{
+		//fixme: auto-start if chosen?
+	}
+
+	virtual void ShutdownModule() override
+	{
+		SetOnScreenDebugRender(false);
+		if (DLCHandler.IsValid())
+		{
+			DLCHandler->Shutdown();
+			DLCHandler.Reset();
+		}
+	}
+
+	virtual TSharedPtr<IPlatformDLC> GetPlatformDLC() override
+	{
+		EnsureInitialized();
+		return DLCHandler;
+	}
+
+
+	void SetOnScreenDebugRender( bool bEnable )
+	{
+#if !UE_BUILD_SHIPPING
+		static FTSTicker::FDelegateHandle DebugDrawTickHandle;
+
+		bool bIsEnabled = (DebugDrawTickHandle.IsValid());
+		if (bIsEnabled == bEnable)
+		{
+			return;
+		}
+
+		if (bEnable)
+		{
+			EnsureInitialized();
+			DebugDrawTickHandle = FTSTicker::GetCoreTicker().AddTicker( TEXT("DLCDebugDraw"), 0.0f, [this](float)
+			{
+				GEngine->AddOnScreenDebugMessage(INDEX_NONE, 0.0f, FColor::White, DLCHandler->Debug_GetStateDescription(), false);
+
+				int PackageIndex = 0;
+				FString PackageDescription;
+				FColor PackageColor;
+				while (DLCHandler->Debug_GetPackageDescription(PackageIndex++, PackageDescription, PackageColor) )
+				{
+					GEngine->AddOnScreenDebugMessage(INDEX_NONE, 0.0f, PackageColor, PackageDescription, false);
+				}
+				return true;
+			});
+		}
+		else
+		{
+			FTSTicker::GetCoreTicker().RemoveTicker(DebugDrawTickHandle);
+			DebugDrawTickHandle.Reset();
+		}
+#endif
+	}
+
+
+	void EnsureInitialized()
+	{
+		// only create GDK DLC handler if the GDK runtime is available
+		if (!DLCHandler.IsValid() && IGDKRuntimeModule::Get().IsAvailable())
+		{
+			DLCHandler = MakeShared<FGDKPlatformDLC>();
+			DLCHandler->Initialize();
+		}
+	}
+
+	TSharedPtr<FGDKPlatformDLC> DLCHandler;
+
+
+
+
+#else
+	// stub when there is no GRDK available
+	virtual TSharedPtr<IPlatformDLC> GetPlatformDLC() override
+	{
+		return nullptr;
+	}
+#endif //WITH_GRDK
+};
+
+
+
+
+
+#if !UE_BUILD_SHIPPING && WITH_GRDK
+static FAutoConsoleCommand CVarGDKDumpDLC(
+	TEXT("GDK.DLC.DumpState"),
+	TEXT("Dumps the state of the DLC"),
+	FConsoleCommandDelegate::CreateLambda([]()
+	{
+		FGDKPlatformDLCModule* Module = FModuleManager::GetModulePtr<FGDKPlatformDLCModule>("GDKPlatformDLC");
+		if (Module != nullptr && Module->DLCHandler.IsValid())
+		{
+			UE_LOG( LogPlatformDLC, Log, TEXT("%s"), *Module->DLCHandler->Debug_GetStateDescription() );
+
+			int PackageIndex = 0;
+			FString PackageDescription;
+			FColor PackageColor;
+			while (Module->DLCHandler->Debug_GetPackageDescription(PackageIndex++, PackageDescription, PackageColor) )
+			{
+				UE_LOG( LogPlatformDLC, Log, TEXT("%s"), *PackageDescription);
+			}
+		}
+	})
+);
+
+
+static FAutoConsoleCommand CVarGDKDLCShow(
+	TEXT("GDK.DLC.DebugShow"),
+	TEXT("Show or hide the on-screen debugging"),
+	FConsoleCommandWithArgsDelegate::CreateLambda( [](const TArray<FString>& Args )
+	{
+		FGDKPlatformDLCModule* Module = FModuleManager::GetModulePtr<FGDKPlatformDLCModule>("GDKPlatformDLC");
+		if (Module != nullptr && Module->DLCHandler.IsValid())
+		{
+			bool bShow = (Args.Num() == 0) || (FCString::Atoi(*Args[0]) != 0);
+			Module->SetOnScreenDebugRender(bShow);
+		}
+	})
+);
+
+#endif //!UE_BUILD_SHIPPING && WITH_GRDK
+
+
+
+IMPLEMENT_MODULE(FGDKPlatformDLCModule, GDKPlatformDLC);
diff --git a/Engine/Plugins/Experimental/Microsoft/GDKPlatformDLC/Source/GDKPlatformDLC/Public/GDKPlatformDLCModule.h b/Engine/Plugins/Experimental/Microsoft/GDKPlatformDLC/Source/GDKPlatformDLC/Public/GDKPlatformDLCModule.h
new file mode 100644
--- /dev/null
+++ b/Engine/Plugins/Experimental/Microsoft/GDKPlatformDLC/Source/GDKPlatformDLC/Public/GDKPlatformDLCModule.h
@@ -0,0 +1,25 @@
+// Copyright Epic Games, Inc. All Rights Reserved.
+
+#pragma once
+
+#include "CoreMinimal.h"
+#include "PlatformDLC.h"
+#include "Modules/ModuleInterface.h"
+#include "Modules/ModuleManager.h"
+
+class GDKPLATFORMDLC_API IGDKPlatformDLCModule : public IPlatformDLCFactoryModule
+{
+public:
+	/**
+	 * Singleton-like access to the MS Gaming Runtime module instance
+	 * @return Returns IGDKPlatformDLCModule singleton instance, loading the module on demand if needed
+	 */
+	static inline IGDKPlatformDLCModule& Get()
+	{
+		return FModuleManager::LoadModuleChecked<IGDKPlatformDLCModule>("GDKPlatformDLC");
+	}
+
+	// get the platform DLC for this factory
+	virtual TSharedPtr<IPlatformDLC> GetPlatformDLC() = 0;
+};
+
diff --git a/Engine/Plugins/Experimental/PlatformDLC/PlatformDLC.uplugin b/Engine/Plugins/Experimental/PlatformDLC/PlatformDLC.uplugin
new file mode 100644
--- /dev/null
+++ b/Engine/Plugins/Experimental/PlatformDLC/PlatformDLC.uplugin
@@ -0,0 +1,25 @@
+{
+	"FileVersion": 3,
+	"Version": 1,
+	"VersionName": "0.1",
+	"FriendlyName": "Platform DLC",
+	"Description": "Experimental support for native platform DLC (on supported platforms)",
+	"CreatedBy": "Epic Games, Inc.",
+	"Category": "Platform",
+	"CreatedByURL": "https://epicgames.com",
+	"DocsURL": "",
+	"MarketplaceURL": "",
+	"SupportURL": "",
+	"EnabledByDefault": false,
+	"RequiresBuildPlatform": false,
+	"CanContainContent": false,
+	"IsBetaVersion": true,
+	"Installed": false,
+	"Modules": [
+		{
+			"Name": "PlatformDLC",
+			"Type": "RuntimeNoCommandlet",
+			"LoadingPhase": "PostConfigInit"
+		}
+	]
+}
\ No newline at end of file
diff --git a/Engine/Plugins/Experimental/PlatformDLC/Source/PlatformDLC/PlatformDLC.Build.cs b/Engine/Plugins/Experimental/PlatformDLC/Source/PlatformDLC/PlatformDLC.Build.cs
new file mode 100644
--- /dev/null
+++ b/Engine/Plugins/Experimental/PlatformDLC/Source/PlatformDLC/PlatformDLC.Build.cs
@@ -0,0 +1,12 @@
+// Copyright Epic Games, Inc. All Rights Reserved.
+
+using UnrealBuildTool;
+using EpicGames.Core;
+
+public class PlatformDLC : ModuleRules
+{
+	public PlatformDLC(ReadOnlyTargetRules Target) : base(Target)
+	{
+		PrivateDependencyModuleNames.Add("Core");
+	}
+}
diff --git a/Engine/Plugins/Experimental/PlatformDLC/Source/PlatformDLC/Private/PlatformDLCModule.cpp b/Engine/Plugins/Experimental/PlatformDLC/Source/PlatformDLC/Private/PlatformDLCModule.cpp
new file mode 100644
--- /dev/null
+++ b/Engine/Plugins/Experimental/PlatformDLC/Source/PlatformDLC/Private/PlatformDLCModule.cpp
@@ -0,0 +1,165 @@
+// Copyright Epic Games, Inc. All Rights Reserved.
+
+#include "PlatformDLCModule.h"
+#include "HAL/IConsoleManager.h"
+#include "HAL/PlatformFile.h"
+#include "HAL/PlatformFileManager.h"
+#include "Misc/ConfigCacheIni.h"
+
+DEFINE_LOG_CATEGORY(LogPlatformDLC);
+
+
+const TCHAR* IPlatformDLC::ConfigSectionName = TEXT("PlatformDLC");
+
+
+class FPlatformDLCModule : public IPlatformDLCModule
+{
+public:
+	virtual TSharedPtr<IPlatformDLC> GetPlatformDLC() override
+	{
+		if (!bHasTried)
+		{
+			// read the configuration & try to create the default platform DLC module
+			FString DefaultModule;
+			if (GConfig->GetString(IPlatformDLC::ConfigSectionName, TEXT("DefaultModule"), DefaultModule, GEngineIni) && !DefaultModule.IsEmpty())
+			{
+				IPlatformDLCFactoryModule* PlatformDLCFactory = FModuleManager::LoadModulePtr<IPlatformDLCFactoryModule>(*DefaultModule);
+				if (PlatformDLCFactory != nullptr)
+				{
+					PlatformDLC = PlatformDLCFactory->GetPlatformDLC();
+				}
+			}
+
+			bHasTried = true;
+		}
+
+		return PlatformDLC.Pin();
+	}
+
+private:
+	bool bHasTried = false;
+	TWeakPtr<IPlatformDLC> PlatformDLC;
+};
+
+
+
+static FAutoConsoleCommand CVarPlatformDLCMount(
+	TEXT("PlatformDLC.Mount"),
+	TEXT("[DLCName] Mount the given platform DLC"),
+	FConsoleCommandWithArgsDelegate::CreateLambda( [](const TArray<FString>& Args )
+	{
+		TSharedPtr<IPlatformDLC> PlatformDLC = IPlatformDLCModule::Get().GetPlatformDLC();
+		if( PlatformDLC.IsValid() && Args.Num() > 0)
+		{
+			bool bResult = PlatformDLC->Mount(FName(Args[0]));
+			UE_LOG(LogPlatformDLC, Log, TEXT("DLC Mount(%s) returned %s"), *Args[0], *LexToString(bResult) ); 
+		}
+	})
+);
+
+static FAutoConsoleCommand CVarPlatformDLCUnmount(
+	TEXT("PlatformDLC.Unmount"),
+	TEXT("[DLCName] Unmount the given platform DLC"),
+	FConsoleCommandWithArgsDelegate::CreateLambda( [](const TArray<FString>& Args )
+	{
+		TSharedPtr<IPlatformDLC> PlatformDLC = IPlatformDLCModule::Get().GetPlatformDLC();
+		if( PlatformDLC.IsValid() && Args.Num() > 0)
+		{
+			bool bResult = PlatformDLC->Unmount(FName(Args[0]));
+			UE_LOG(LogPlatformDLC, Log, TEXT("DLC Unmount(%s) returned %s"), *Args[0], *LexToString(bResult) ); 
+		}
+	})
+);
+
+static FAutoConsoleCommand CVarPlatformDLCDownload(
+	TEXT("PlatformDLC.Download"),
+	TEXT("[DLCName] Download the given DLC"),
+	FConsoleCommandWithArgsDelegate::CreateLambda( [](const TArray<FString>& Args )
+	{
+		TSharedPtr<IPlatformDLC> PlatformDLC = IPlatformDLCModule::Get().GetPlatformDLC();
+		if( PlatformDLC.IsValid() && Args.Num() > 0)
+		{
+			bool bResult = PlatformDLC->Download(FName(Args[0]));
+			UE_LOG(LogPlatformDLC, Log, TEXT("DLC Download(%s) returned %s"), *Args[0], *LexToString(bResult) );
+		}
+	})
+);
+
+static FAutoConsoleCommand CVarPlatformDLCUninstall(
+	TEXT("PlatformDLC.Uninstall"),
+	TEXT("[DLCName] Uninstalls the given DLC (debug only)"),
+	FConsoleCommandWithArgsDelegate::CreateLambda( [](const TArray<FString>& Args )
+	{
+		TSharedPtr<IPlatformDLC> PlatformDLC = IPlatformDLCModule::Get().GetPlatformDLC();
+		if( PlatformDLC.IsValid() && Args.Num() > 0)
+		{
+			bool bResult = PlatformDLC->Uninstall(FName(Args[0]));
+			UE_LOG(LogPlatformDLC, Log, TEXT("Uninstal(%s) returned %s"), *Args[0], *LexToString(bResult) );
+		}
+	})
+);
+
+static FAutoConsoleCommand CVarPlatformDLCDir(
+	TEXT("PlatformDLC.Debug.Dir"),
+	TEXT("[DLCName] Dump all files in the given DLC (debug only)"),
+	FConsoleCommandWithArgsDelegate::CreateLambda( [](const TArray<FString>& Args )
+	{
+		TSharedPtr<IPlatformDLC> PlatformDLC = IPlatformDLCModule::Get().GetPlatformDLC();
+		if( PlatformDLC.IsValid() && Args.Num() > 0)
+		{
+			FName DLCName = FName(Args[0]);
+			if (PlatformDLC->IsMounted(DLCName))
+			{
+				FString MountPath = PlatformDLC->GetRootDirectory(DLCName);
+				if (!MountPath.IsEmpty())
+				{
+					UE_LOG(LogPlatformDLC, Log, TEXT("All Files For DLC '%s' in Mount Point '%s'"), *Args[0], *MountPath);
+	
+					class FDumpFileVistor : public IPlatformFile::FDirectoryVisitor
+					{
+					public:
+						virtual bool Visit(const TCHAR* FilenameOrDirectory, bool bIsDirectory)
+						{
+							UE_CLOG(!bIsDirectory, LogPlatformDLC, Log, TEXT("%s"), FilenameOrDirectory);
+							return true;
+						}
+					};
+					FDumpFileVistor Visitor;
+					IPlatformFile::GetPlatformPhysical().IterateDirectoryRecursively(*MountPath, Visitor);
+				}
+				else
+				{
+					UE_LOG(LogPlatformDLC, Error, TEXT("DLC '%s' has no file system mount point"), *Args[0]);
+				}
+			}
+			else
+			{
+				UE_LOG(LogPlatformDLC, Error, TEXT("DLC '%s' is not mounted"), *Args[0]);
+			}
+		}
+	})
+);
+
+
+static FAutoConsoleCommand CVarPlatformDLCList(
+	TEXT("PlatformDLC.Debug.List"),
+	TEXT("Lists all known DLC names (debug only)"),
+	FConsoleCommandDelegate::CreateLambda([]()
+	{
+		TSharedPtr<IPlatformDLC> PlatformDLC = IPlatformDLCModule::Get().GetPlatformDLC();
+		if( PlatformDLC.IsValid())
+		{
+			TArray<FName> DLCNames = PlatformDLC->GetAllDLCNames();
+			UE_LOG(LogPlatformDLC, Log, TEXT("%d DLCs"), DLCNames.Num() );
+			for (const FName& DLCName : DLCNames)
+			{
+				UE_LOG(LogPlatformDLC, Log, TEXT("    %-32s Downloaded: %-5s Mounted:%-5s Entitled:%-5s StoreId:%s"), *DLCName.ToString(), *LexToString(PlatformDLC->IsDownloaded(DLCName)), *LexToString(PlatformDLC->IsMounted(DLCName)), *LexToString(PlatformDLC->HasEntitlement(DLCName)), *PlatformDLC->GetStoreId(DLCName) );
+			}
+		}
+	})
+);
+
+
+
+
+IMPLEMENT_MODULE(FPlatformDLCModule, PlatformDLC);
diff --git a/Engine/Plugins/Experimental/PlatformDLC/Source/PlatformDLC/Public/PlatformDLC.h b/Engine/Plugins/Experimental/PlatformDLC/Source/PlatformDLC/Public/PlatformDLC.h
new file mode 100644
--- /dev/null
+++ b/Engine/Plugins/Experimental/PlatformDLC/Source/PlatformDLC/Public/PlatformDLC.h
@@ -0,0 +1,93 @@
+// Copyright Epic Games, Inc. All Rights Reserved.
+
+#pragma once
+
+#include "CoreMinimal.h"
+#include "Delegates/Delegate.h"
+#include "Modules/ModuleInterface.h"
+
+// main interface for Platform DLC
+// note: API may be subject to change while this plugin is experimental
+class IPlatformDLC
+{
+public:
+	virtual ~IPlatformDLC() = default;
+
+	enum class ENotification
+	{
+		Entitlement, // true = entitlement gained, false = lost
+		Mounted,     // whether the DLC has been successfully mounted and can now be used
+		Unmounted,   // whether the DLC has been successfully unmounted
+		Downloaded,  // whether the DLC has been successfully downloaded and can now be mounted (note if AutoMount is configured, this will happen automatically)
+		Uninstalled, // whether the DLC has been successfully uninstalled
+	};
+	DECLARE_MULTICAST_DELEGATE_ThreeParams(FOnDLCNotification, FName, ENotification, bool/*success*/ )
+
+
+	// "initialize te platform DLC
+	virtual void Initialize() = 0;
+
+	// shutdown the platform DLC and clean up resources
+	virtual void Shutdown() = 0;
+
+	// determines whether the platform DLC has finished initializing
+	virtual bool IsInitialized() const = 0;
+
+
+	// query if we have entitlement to download & mount the given DLC
+	virtual bool HasEntitlement( FName DLCName ) const = 0;
+
+	// determines whether the given DLC is mounted
+	virtual bool IsMounted( FName DLCName ) const = 0;
+
+	// determines whether the given DLC is downloaded
+	virtual bool IsDownloaded( FName DLCName ) const = 0;
+
+
+	// called when there is a notification about DLC, such as mounting completion
+	virtual FOnDLCNotification& OnNotification() = 0;
+
+
+	// start asynchronously mounting the given DLC
+	// can also be used to attempt entitlement reacquisition after it has been lost
+	// returns true if the callback will be triggered, or false if the DLC is not known
+	virtual bool Mount( FName DLCName ) = 0;
+
+	// start asynchronously unmounting the given DLC. It is essential that the game is no longer using the DLC
+	// returns true if the callback will be triggered, or false if the DLC is not known
+	virtual bool Unmount( FName DLCName ) = 0;
+
+	// start asynchronously downloading the given DLC. Depending on the platform you may reacquire entitlements before you can download.
+	// returns true if the callback will be triggered, or false if the DLC is not known
+	virtual bool Download( FName DLCName ) = 0;
+
+	// start uninstalling the given DLC (may not work properly on all platforms. DLC must be unmounted first)
+	// returns true if the callback will be triggered, or false if the DLC is not known
+	virtual bool Uninstall( FName DLCName ) = 0;
+
+
+	// returns the file system mount point for the given DLC
+	virtual FString GetRootDirectory( FName DLCName ) const = 0;
+
+	// returns the list of all known DLC (regardless of whether downloaded / mounted)
+	virtual TArray<FName> GetAllDLCNames() const = 0;
+
+	// returns the list of all DLC that is ready for use (entitled, downloaded & mounted)
+	virtual TArray<FName> GetMountedDLCNames() const = 0;
+
+	// get the store id for the DLC (should be convertible to FUniqueOfferId)
+	virtual FString GetStoreId( FName DLCName ) const = 0;
+
+
+	// config section for Platform DLC settings
+	static PLATFORMDLC_API const TCHAR* ConfigSectionName;
+};
+
+
+// interface for a platform DLC factory
+class IPlatformDLCFactoryModule : public IModuleInterface
+{
+public:
+	// get the platform DLC for this factory
+	virtual TSharedPtr<IPlatformDLC> GetPlatformDLC() = 0;
+};
diff --git a/Engine/Plugins/Experimental/PlatformDLC/Source/PlatformDLC/Public/PlatformDLCModule.h b/Engine/Plugins/Experimental/PlatformDLC/Source/PlatformDLC/Public/PlatformDLCModule.h
new file mode 100644
--- /dev/null
+++ b/Engine/Plugins/Experimental/PlatformDLC/Source/PlatformDLC/Public/PlatformDLCModule.h
@@ -0,0 +1,30 @@
+// Copyright Epic Games, Inc. All Rights Reserved.
+
+#pragma once
+
+#include "CoreMinimal.h"
+#include "PlatformDLC.h"
+#include "Modules/ModuleInterface.h"
+#include "Modules/ModuleManager.h"
+#include "Templates/SharedPointer.h"
+
+PLATFORMDLC_API DECLARE_LOG_CATEGORY_EXTERN(LogPlatformDLC, Log, All)
+
+class IPlatformDLCModule : public IModuleInterface
+{
+public:
+	/**
+	 * Singleton-like access to the platform DLC module instance
+	 * @return Returns IPlatformDLCModule singleton instance, loading the module on demand if needed
+	 */
+	static inline IPlatformDLCModule& Get()
+	{
+		return FModuleManager::LoadModuleChecked<IPlatformDLCModule>("PlatformDLC");
+	}
+
+	/**
+	 * Returns the platform DLC, if any (configured via ini)
+	 * @return platform DLC pointer, or null if it isn't available
+	 */
+	virtual TSharedPtr<IPlatformDLC> GetPlatformDLC() = 0;
+};
