don't create DLC handler if this isn't a packaged build

#rnx

#Synced-CL 51451766

diff --git a/Engine/Plugins/Experimental/Microsoft/GDKPlatformDLC/Source/GDKPlatformDLC/Private/GDKPlatformDLCModule.cpp b/Engine/Plugins/Experimental/Microsoft/GDKPlatformDLC/Source/GDKPlatformDLC/Private/GDKPlatformDLCModule.cpp
--- a/Engine/Plugins/Experimental/Microsoft/GDKPlatformDLC/Source/GDKPlatformDLC/Private/GDKPlatformDLCModule.cpp
+++ b/Engine/Plugins/Experimental/Microsoft/GDKPlatformDLC/Source/GDKPlatformDLC/Private/GDKPlatformDLCModule.cpp
@@ -11,6 +11,11 @@
 #include "HAL/IConsoleManager.h"
 #include "Engine/Engine.h"
 #include "Containers/Ticker.h"
+#include "Microsoft/AllowMicrosoftPlatformTypes.h"
+    THIRD_PARTY_INCLUDES_START
+    #include <XPackage.h>
+    THIRD_PARTY_INCLUDES_END
+#include "Microsoft/HideMicrosoftPlatformTypes.h"
 #endif
 
 class FGDKPlatformDLCModule : public IGDKPlatformDLCModule
@@ -56,14 +61,17 @@ public:
 			EnsureInitialized();
 			DebugDrawTickHandle = FTSTicker::GetCoreTicker().AddTicker( TEXT("DLCDebugDraw"), 0.0f, [this](float)
 			{
-				GEngine->AddOnScreenDebugMessage(INDEX_NONE, 0.0f, FColor::White, DLCHandler->Debug_GetStateDescription(), false);
-
-				int PackageIndex = 0;
-				FString PackageDescription;
-				FColor PackageColor;
-				while (DLCHandler->Debug_GetPackageDescription(PackageIndex++, PackageDescription, PackageColor) )
+				if (DLCHandler.IsValid())
 				{
-					GEngine->AddOnScreenDebugMessage(INDEX_NONE, 0.0f, PackageColor, PackageDescription, false);
+					GEngine->AddOnScreenDebugMessage(INDEX_NONE, 0.0f, FColor::White, DLCHandler->Debug_GetStateDescription(), false);
+
+					int PackageIndex = 0;
+					FString PackageDescription;
+					FColor PackageColor;
+					while (DLCHandler->Debug_GetPackageDescription(PackageIndex++, PackageDescription, PackageColor) )
+					{
+						GEngine->AddOnScreenDebugMessage(INDEX_NONE, 0.0f, PackageColor, PackageDescription, false);
+					}
 				}
 				return true;
 			});
@@ -80,7 +88,7 @@ public:
 	void EnsureInitialized()
 	{
 		// only create GDK DLC handler if the GDK runtime is available
-		if (!DLCHandler.IsValid() && IGDKRuntimeModule::Get().IsAvailable())
+		if (!DLCHandler.IsValid() && IGDKRuntimeModule::Get().IsAvailable() && XPackageIsPackagedProcess())
 		{
 			DLCHandler = MakeShared<FGDKPlatformDLC>();
 			DLCHandler->Initialize();
