Add stub function for GetDownloadSize (not implemented yet)

#rnx
#Synced-CL 51451766

diff --git a/Engine/Plugins/Experimental/Microsoft/GDKPlatformDLC/Source/GDKPlatformDLC/Private/GDKPlatformDLC.cpp b/Engine/Plugins/Experimental/Microsoft/GDKPlatformDLC/Source/GDKPlatformDLC/Private/GDKPlatformDLC.cpp
--- a/Engine/Plugins/Experimental/Microsoft/GDKPlatformDLC/Source/GDKPlatformDLC/Private/GDKPlatformDLC.cpp
+++ b/Engine/Plugins/Experimental/Microsoft/GDKPlatformDLC/Source/GDKPlatformDLC/Private/GDKPlatformDLC.cpp
@@ -803,6 +803,23 @@ bool FGDKPlatformDLC::Mount( FName DLCName )
 }
 
 
+bool FGDKPlatformDLC::GetDownloadSize( FName DLCName, uint64& OutCurrentInstallSize, uint64& OutFullInstallSize ) const
+{
+	// @fixme: stub implementation!!
+	// need XPackageGetInstallationProgress and a XPackageInstallationMonitorHandle per-package
+
+	if (!DLCNameToStoreIdMap.Contains(DLCName))
+	{
+		OutFullInstallSize = 0;	
+		OutCurrentInstallSize = 0;	
+		return false;
+	}
+
+	OutFullInstallSize = 1024*1024;
+	OutCurrentInstallSize = IsDownloaded(DLCName) ? OutFullInstallSize : 0;
+	return true;
+}
+
 FString FGDKPlatformDLC::GetRootDirectory( FName DLCName ) const
 {
 	TSharedPtr<const FPackageData> Package = FindDLC(DLCName);
diff --git a/Engine/Plugins/Experimental/Microsoft/GDKPlatformDLC/Source/GDKPlatformDLC/Private/GDKPlatformDLC.h b/Engine/Plugins/Experimental/Microsoft/GDKPlatformDLC/Source/GDKPlatformDLC/Private/GDKPlatformDLC.h
--- a/Engine/Plugins/Experimental/Microsoft/GDKPlatformDLC/Source/GDKPlatformDLC/Private/GDKPlatformDLC.h
+++ b/Engine/Plugins/Experimental/Microsoft/GDKPlatformDLC/Source/GDKPlatformDLC/Private/GDKPlatformDLC.h
@@ -41,6 +41,7 @@ public:
 	virtual bool Download( FName DLCName ) override;
 	virtual bool Uninstall( FName DLCName ) override;
 
+	virtual bool GetDownloadSize( FName DLCName, uint64& OutCurrentInstallSize, uint64& OutFullInstallSize ) const override;
 	virtual FString GetRootDirectory( FName DLCName ) const override;
 	virtual TArray<FName> GetAllDLCNames() const override;
 	virtual TArray<FName> GetMountedDLCNames() const override;
diff --git a/Engine/Plugins/Experimental/PlatformDLC/Source/PlatformDLC/Public/PlatformDLC.h b/Engine/Plugins/Experimental/PlatformDLC/Source/PlatformDLC/Public/PlatformDLC.h
--- a/Engine/Plugins/Experimental/PlatformDLC/Source/PlatformDLC/Public/PlatformDLC.h
+++ b/Engine/Plugins/Experimental/PlatformDLC/Source/PlatformDLC/Public/PlatformDLC.h
@@ -66,6 +66,10 @@ public:
 	virtual bool Uninstall( FName DLCName ) = 0;
 
 
+	// determine the download size for the given DLC
+	// returns true if the download size is available
+	virtual bool GetDownloadSize( FName DLCName, uint64& OutCurrentInstallSize, uint64& OutFullInstallSize ) const = 0;
+
 	// returns the file system mount point for the given DLC
 	virtual FString GetRootDirectory( FName DLCName ) const = 0;
 
