Added download progress tracking to GDK Platform DLC
- also includes workaround for potential bug when re-mounting previously uninstalled DLC

#rnx
#Synced-CL 51451766

diff --git a/Engine/Plugins/Experimental/Microsoft/GDKPlatformDLC/Source/GDKPlatformDLC/Private/GDKPlatformDLC.cpp b/Engine/Plugins/Experimental/Microsoft/GDKPlatformDLC/Source/GDKPlatformDLC/Private/GDKPlatformDLC.cpp
--- a/Engine/Plugins/Experimental/Microsoft/GDKPlatformDLC/Source/GDKPlatformDLC/Private/GDKPlatformDLC.cpp
+++ b/Engine/Plugins/Experimental/Microsoft/GDKPlatformDLC/Source/GDKPlatformDLC/Private/GDKPlatformDLC.cpp
@@ -21,6 +21,13 @@
 #include "PlatformDLCModule.h"
 
 
+static bool GGDKDLCAllowRedownloadAfterMount = false;
+static FAutoConsoleVariableRef CVarDLCAllowRedownloadAfterMount(
+	TEXT( "GDK.DLC.AllowRedownloadAfterMount" ),
+	GGDKDLCAllowRedownloadAfterMount,
+	TEXT( "Whether to allow DLC to be downloaded again after it has previously been mounted. Disabled by default due to a bug that causes device I/O errors when re-mounting DLC that was previously uninstalled" ));
+
+
 FGDKPlatformDLC::FGDKPlatformDLC()
 	: bIsInitialized(false)
 	, CurrentStoreContext(nullptr)
@@ -74,6 +81,9 @@ void FGDKPlatformDLC::Initialize()
 			if (FParse::Value(*DLCPlugin, TEXT("DLCName="), DLCName) && FParse::Value(*DLCPlugin, TEXT("StoreId="), StoreId))
 			{
 				DLCNameToStoreIdMap.Add(FName(*DLCName), StoreId);
+
+				TSharedPtr<FPackageData> Package = AllPackages.Add_GetRef(MakeShared<FPackageData>(this, FName(*DLCName), StoreId));
+				Package->Context.Package = Package;
 			}
 		}
 	}
@@ -195,32 +205,30 @@ void FGDKPlatformDLC::AsyncInitDLC(const XPackageDetails& Details, bool bMount)
 {
 	UE_CLOG(Details.installing, LogPlatformDLC, Error, TEXT("package %s is still installing - intelligent delivery is not supported for DLC at the moment"), UTF8_TO_TCHAR(Details.packageIdentifier) );
 
-	// add package and prepare for enumeration callbacks
 	TSharedPtr<FPackageData> Package;
 	{
 		FScopeLock Lock(&PackageDataLock);
-		
-		FString StoreId = UTF8_TO_TCHAR(Details.storeId);
-		FName DLCName;
-		for (const TTuple<FName,FString>& Pair : DLCNameToStoreIdMap)
+
+		// find the package by store id
+		for (const TSharedPtr<FPackageData>& ExistingPackage : AllPackages)
 		{
-			if (Pair.Value == StoreId)
+			if (FCStringAnsi::Strcmp(ExistingPackage->StoreId, Details.storeId) == 0)
 			{
-				DLCName = Pair.Key;
+				Package = ExistingPackage;
 				break;
 			}
 		}
-		if (DLCName.IsNone())
+
+		if (!Package.IsValid())
 		{
-			UE_LOG(LogPlatformDLC, Warning, TEXT("Ignoring DLC with StoreId %s because it isn't referenced in DLCPlugins config array"), *StoreId);
+			UE_LOG(LogPlatformDLC, Warning, TEXT("Ignoring DLC with StoreId %s because it isn't referenced in DLCPlugins config array"), UTF8_TO_TCHAR(Details.storeId));
 			return;
 		}
 
-		Package = AllPackages.Add_GetRef( MakeShared<FPackageData>(this, Details, DLCName));
-		Package->Context.Package = Package; // initialize the callback context. necessary to support those GDK callback functions that only take a raw context pointer
+		Package->OnDownloaded(Details);
+		PostNotification(Package->DLCName, ENotification::Downloaded, true);
 	}
 
-	// mount the DLC package automatically, if requested (this is typically done for DLC that arrives after engine startup)
 	if (bMount)
 	{
 		AddToPendingPackages( Package );
@@ -228,11 +236,19 @@ void FGDKPlatformDLC::AsyncInitDLC(const XPackageDetails& Details, bool bMount)
 }
 
 
-void FGDKPlatformDLC::AddToPendingPackages( TSharedPtr<FPackageData> Package )
+bool FGDKPlatformDLC::AddToPendingPackages( TSharedPtr<FPackageData> Package )
 {
+	if (Package->bWasPreviouslyMounted && !GGDKDLCAllowRedownloadAfterMount)
+	{
+		UE_LOG(LogPlatformDLC, Error, TEXT("DLC %s cannot currently be re-mounted after being previously uninstalled. Please restart the application to re-mount this DLC. Set GDK.DLC.AllowRedownloadAfterMount=1 to override."), *Package->DLCName.ToString());
+		return false;
+	}
+
 	FScopeLock Lock(&PackageDataLock);
 	PendingPackages.AddUnique(Package);
 	AddToTicker();
+
+	return true;
 }
 
 
@@ -564,6 +580,7 @@ void FGDKPlatformDLC::MountPlugin_GameThread( TSharedPtr<FPackageData> Package )
 
 	// all done
 	Package->MountState = EMountState::Mounted;
+	Package->bWasPreviouslyMounted = true;
 	ProcessingState = EProcessingState::Completed;
 }
 
@@ -654,6 +671,8 @@ bool FGDKPlatformDLC::Unmount( FName DLCName )
 			}
 		}
 
+		IGDKPackageManifestModule::Get().UnloadDLC(Package->DLCHandle);
+
 		// unmount the package
 		XPackageCloseMountHandle(Package->DLCHandle);
 		Package->DLCHandle = nullptr;
@@ -694,6 +713,22 @@ void FGDKPlatformDLC::RemoveFromTicker()
 
 
 
+void FGDKPlatformDLC::OnInstallationProgressChanged(void* Context, XPackageInstallationMonitorHandle Monitor)
+{
+	TSharedPtr<FPackageData> Package = ((FPackageData::FContext*)Context)->Package.Pin();
+	if (Package.IsValid())
+	{
+		XPackageInstallationProgress Progress;
+		XPackageGetInstallationProgress(Monitor, &Progress);
+		Package->CachedInstalledBytes = Progress.installedBytes;
+		Package->CachedTotalBytes = Progress.totalBytes;
+
+		const float Percent = (Progress.totalBytes > 0) ? 100.0f * (float)Progress.installedBytes / (float)Progress.totalBytes : 0.0f;
+		UE_LOG(LogPlatformDLC, Verbose, TEXT("DLC %s download progress changed: %20llu / %20llu bytes  (%6.2f%%)  %s"), *Package->DLCName.ToString(), Progress.installedBytes, Progress.totalBytes, Percent, Progress.completed ? TEXT("Completed") : TEXT(""));
+	}
+}
+
+
 void FGDKPlatformDLC::OnPackageLicenseLost(void* Context)
 {
 	TSharedPtr<FPackageData> Package = ((FPackageData::FContext*)Context)->Package.Pin();
@@ -783,9 +818,8 @@ bool FGDKPlatformDLC::IsMounted( FName DLCName ) const
 
 bool FGDKPlatformDLC::IsDownloaded( FName DLCName ) const
 {
-	// FPackageData only exists if it's downloaded
 	TSharedPtr<const FPackageData> Package = FindDLC(DLCName);
-	return (Package.IsValid());
+	return Package.IsValid() && Package->DownloadState == EDownloadState::Downloaded;
 }
 
 
@@ -797,26 +831,29 @@ bool FGDKPlatformDLC::Mount( FName DLCName )
 		return false;
 	}
 
-	AddToPendingPackages(Package);
-
-	return true;
+	return AddToPendingPackages(Package);
 }
 
 
 bool FGDKPlatformDLC::GetDownloadSize( FName DLCName, uint64& OutCurrentInstallSize, uint64& OutFullInstallSize ) const
 {
-	// @fixme: stub implementation!!
-	// need XPackageGetInstallationProgress and a XPackageInstallationMonitorHandle per-package
+	OutCurrentInstallSize = 0;
+	OutFullInstallSize = 0;
+
+	TSharedPtr<const FPackageData> Package = FindDLC(DLCName);
+	if (!Package.IsValid())
+	{
+		return false;
+	}
 
-	if (!DLCNameToStoreIdMap.Contains(DLCName))
+	// CachedTotalBytes is 0 before a download has started, and UINT64_MAX until the GDK resolves the package size
+	if (Package->CachedTotalBytes == 0 || Package->CachedTotalBytes == UINT64_MAX)
 	{
-		OutFullInstallSize = 0;	
-		OutCurrentInstallSize = 0;	
 		return false;
 	}
 
-	OutFullInstallSize = 1024*1024;
-	OutCurrentInstallSize = IsDownloaded(DLCName) ? OutFullInstallSize : 0;
+	OutCurrentInstallSize = Package->CachedInstalledBytes;
+	OutFullInstallSize = Package->CachedTotalBytes;
 	return true;
 }
 
@@ -882,30 +919,21 @@ void FGDKPlatformDLC::PostNotification(FName DLCName, ENotification Notification
 
 
 
-FGDKPlatformDLC::FPackageData::FPackageData(FGDKPlatformDLC* InOwner, const XPackageDetails& InDetails, const FName& InDLCName)
-	: MountState(EMountState::NotMounted)
-	, LicenseState(ELicenseState::Unlicensed)
-	, Owner(InOwner)
-	, bIsAgeRestricted(InDetails.ageRestricted)
+FGDKPlatformDLC::FPackageData::FPackageData(FGDKPlatformDLC* InOwner, const FName& InDLCName, const FString& InStoreId)
+	: Owner(InOwner)
 	, DLCName(InDLCName)
-	, DLCHandle(nullptr)
-	, DLCLicense(nullptr)
 {
-	FCStringAnsi::Strncpy(PackageIdentifier, InDetails.packageIdentifier, XPACKAGE_IDENTIFIER_MAX_LENGTH);
-	FCStringAnsi::Strncpy(StoreId, InDetails.storeId, STORE_SKU_ID_SIZE);
-
-	// set the name
-	PackageName = UTF8_TO_TCHAR(InDetails.displayName);
-	if (PackageName.IsEmpty())
-	{
-		PackageName = UTF8_TO_TCHAR(PackageIdentifier);
-	}
+	FCStringAnsi::Strncpy(StoreId, TCHAR_TO_UTF8(*InStoreId), STORE_SKU_ID_SIZE);
+	PackageName = InDLCName.ToString(); // placeholder until the display name is known
 }
 
 
 
+
 FGDKPlatformDLC::FPackageData::~FPackageData()
 {
+	DestroyInstallationMonitor();
+
 	if (DLCHandle != nullptr)
 	{
 		XPackageCloseMountHandle(DLCHandle);
@@ -920,6 +948,81 @@ FGDKPlatformDLC::FPackageData::~FPackageData()
 	}
 }
 
+void FGDKPlatformDLC::FPackageData::OnDownloading(const char* InPackageIdentifier)
+{
+	// store package details
+	FCStringAnsi::Strncpy(PackageIdentifier, InPackageIdentifier, XPACKAGE_IDENTIFIER_MAX_LENGTH);
+	
+	DownloadState = EDownloadState::Downloading;
+	CreateInstallationMonitor(PackageIdentifier);
+}
+
+
+void FGDKPlatformDLC::FPackageData::OnDownloaded(const XPackageDetails& InDetails)
+{
+	// store package details
+	FCStringAnsi::Strncpy(PackageIdentifier, InDetails.packageIdentifier, XPACKAGE_IDENTIFIER_MAX_LENGTH);
+	bIsAgeRestricted = InDetails.ageRestricted;
+	if (InDetails.displayName[0] != '\0')
+	{
+		PackageName = UTF8_TO_TCHAR(InDetails.displayName);
+	}
+
+	DownloadState = EDownloadState::Downloaded;
+	DestroyInstallationMonitor();
+}
+
+
+void FGDKPlatformDLC::FPackageData::OnUninstalled()
+{
+	// clear package details 
+	PackageIdentifier[0] = '\0';
+	MountPath.Reset();
+	DLCPlugin.Reset();
+	CachedInstalledBytes = 0; // not clearing CachedTotalBytes too because that will not change on a subsequent re-download
+	MountState = EMountState::NotMounted;
+	
+	DownloadState = EDownloadState::NotDownloaded;
+	DestroyInstallationMonitor();	
+}
+
+
+void FGDKPlatformDLC::FPackageData::CreateInstallationMonitor(const char* Identifier)
+{
+	if (InstallationMonitorHandle == nullptr)
+	{
+		HRESULT hResult = XPackageCreateInstallationMonitor(Identifier, 0, nullptr, 500, FGDKAsyncTaskQueue::GetGenericQueue(), &InstallationMonitorHandle);
+		if (SUCCEEDED(hResult))
+		{
+			XPackageInstallationProgress Progress;
+			XPackageGetInstallationProgress(InstallationMonitorHandle, &Progress);
+			CachedInstalledBytes = Progress.installedBytes;
+			CachedTotalBytes = Progress.totalBytes;
+
+			hResult = XPackageRegisterInstallationProgressChanged(InstallationMonitorHandle, &Context, &FGDKPlatformDLC::OnInstallationProgressChanged, &InstallationProgressToken);
+			UE_CLOG(FAILED(hResult), LogPlatformDLC, Warning, TEXT("Failed to register installation progress callback for %s : 0x%X"), *PackageName, hResult);
+		}
+		else
+		{
+			UE_LOG(LogPlatformDLC, Warning, TEXT("Failed to create installation monitor for %s : 0x%X"), *PackageName, hResult);
+		}
+	}
+}
+
+
+void FGDKPlatformDLC::FPackageData::DestroyInstallationMonitor()
+{
+	if (InstallationMonitorHandle != nullptr)
+	{
+		XPackageUnregisterInstallationProgressChanged(InstallationMonitorHandle, InstallationProgressToken, true);
+		InstallationProgressToken = {};
+		
+		XPackageCloseInstallationMonitorHandle(InstallationMonitorHandle);
+		InstallationMonitorHandle = nullptr;
+	}
+}
+
+
 
 
 
@@ -932,11 +1035,11 @@ bool FGDKPlatformDLC::Uninstall( FName DLCName )
 	{
 		if (Package->DLCHandle == nullptr)
 		{
-			bool bSuccess = XPackageUninstallPackage(TCHAR_TO_UTF8(Package->PackageIdentifier));
+			bool bSuccess = XPackageUninstallPackage(Package->PackageIdentifier);
 			if (bSuccess)
 			{
 				FScopeLock Lock(&PackageDataLock);
-				AllPackages.Remove(Package);
+				Package->OnUninstalled();
 			}
 			PostNotification(Package->DLCName, ENotification::Uninstalled, bSuccess);
 			return true;
@@ -955,78 +1058,83 @@ bool FGDKPlatformDLC::Uninstall( FName DLCName )
 bool FGDKPlatformDLC::Download( FName DLCName )
 {
 	TSharedPtr<FGDKPlatformDLC::FPackageData> Package = FindDLC(DLCName);
-	if (Package.IsValid())
+	if (!Package.IsValid())
+	{
+		return false;
+	}
+
+	if (Package->DownloadState == EDownloadState::Downloaded)
 	{
-		// if we know about the package, it's already downloaded
 		PostNotification(DLCName, ENotification::Downloaded, true);
 		return true;
 	}
-	else
-	{
-		// find the package with the store id
-		FString StoreId = GetStoreId(DLCName);
-		if (!StoreId.IsEmpty())
-		{
-			// make sure there is a store context
-			if (!EnsureValidStoreContext())
-			{
-				return false;
-			}
 
-			// kick off the async 
-			HRESULT hResult = AsyncGDKTask( 
-				[this,StoreId]( XAsyncBlock* AsyncBlock )
-				{
-					char StoreIdAnsi[64];
-					FCStringAnsi::Strncpy(StoreIdAnsi, TCHAR_TO_ANSI(*StoreId), UE_ARRAY_COUNT(StoreIdAnsi));
+	if (Package->DownloadState == EDownloadState::Downloading)
+	{
+		return true;
+	}
 
-					const char* StoreIds[] =
-					{
-						StoreIdAnsi
-					};
+	if (Package->bWasPreviouslyMounted && !GGDKDLCAllowRedownloadAfterMount)
+	{
+		UE_LOG(LogPlatformDLC, Error, TEXT("DLC %s cannot currently be re-downloaded after being mounted. Please restart the application to re-acquire this DLC. Set GDK.DLC.AllowRedownloadAfterMount=1 to override."), *DLCName.ToString());
+		return false;
+	}
 
-					return XStoreDownloadAndInstallPackagesAsync( CurrentStoreContext, StoreIds, 1, AsyncBlock );
-				},
-				[this, DLCName]( XAsyncBlock* AsyncBlock )
-				{
-					HRESULT hResult;
+	// make sure there is a store context
+	if (!EnsureValidStoreContext())
+	{
+		return false;
+	}
 
-					uint32 NumPackages = 0;
-					hResult = XStoreDownloadAndInstallPackagesResultCount( AsyncBlock, &NumPackages );
-					if (FAILED(hResult))
-					{
-						UE_LOG(LogPlatformDLC, Error, TEXT("failed to get installing packages count for %s : 0x%X"), *DLCName.ToString(), hResult );
-						PostNotification(DLCName, ENotification::Downloaded, false);
-						return;
-					}
+	// kick off the async download
+	HRESULT hResult = AsyncGDKTask(
+		[this, Package]( XAsyncBlock* AsyncBlock )
+		{
+			const char* StoreIds[] = 
+			{
+				Package->StoreId
+			};
 
-					auto PackageIdentifiers = new char[NumPackages][XPACKAGE_IDENTIFIER_MAX_LENGTH];
-					hResult = XStoreDownloadAndInstallPackagesResult( AsyncBlock, NumPackages, &PackageIdentifiers[0] );
-					if (FAILED(hResult))
-					{
-						delete[] PackageIdentifiers;
-						UE_LOG(LogPlatformDLC, Error, TEXT("failed to get installing packages identifiers for %s : 0x%X"), *DLCName.ToString(), hResult );
-						PostNotification(DLCName, ENotification::Downloaded, false);
-						return;
-					}
+			return XStoreDownloadAndInstallPackagesAsync( CurrentStoreContext, StoreIds, 1, AsyncBlock );
+		},
+		[this, DLCName, Package]( XAsyncBlock* AsyncBlock )
+		{
+			HRESULT hResult;
 
-					// the download was successfully started. we must have a valid license for this before we can mount it
-					for (uint32 Index = 0; Index < NumPackages; Index++)
-					{
-						UE_LOG(LogPlatformDLC, Log, TEXT("Downloading package %s for DLC %s"), UTF8_TO_TCHAR(PackageIdentifiers[Index]), *DLCName.ToString() );
-					}
-	
-					delete[] PackageIdentifiers;
+			uint32 NumPackages = 0;
+			hResult = XStoreDownloadAndInstallPackagesResultCount( AsyncBlock, &NumPackages );
+			if (FAILED(hResult) || NumPackages == 0)
+			{
+				UE_LOG(LogPlatformDLC, Error, TEXT("failed to get installing packages count for %s : 0x%X"), *DLCName.ToString(), hResult );
+				PostNotification(DLCName, ENotification::Downloaded, false);
+				return;
+			}
+			if (NumPackages > 1) // download will still trigger XPackageRegisterPackageInstalled
+			{
+				UE_LOG(LogPlatformDLC, Error, TEXT("expected to get 1 package being downloaded for %s, but got %d. Installation progress will not be tracked"), *DLCName.ToString(), NumPackages );
+				return;
+			}
 
-					PostNotification(DLCName, ENotification::Downloaded, true);
-				}
-			);
+			char PackageIdentifier[XPACKAGE_IDENTIFIER_MAX_LENGTH];
+			hResult = XStoreDownloadAndInstallPackagesResult( AsyncBlock, 1, &PackageIdentifier );
+			if (FAILED(hResult))
+			{
+				UE_LOG(LogPlatformDLC, Error, TEXT("failed to get installing package identifier for %s : 0x%X"), *DLCName.ToString(), hResult );
+				PostNotification(DLCName, ENotification::Downloaded, false);
+				return;
+			}
 
-			return SUCCEEDED(hResult);
-		}
+			// the download was successfully started. we must have a valid license for this before we can mount it
+			UE_LOG(LogPlatformDLC, Log, TEXT("Downloading package %s for DLC %s"), UTF8_TO_TCHAR(PackageIdentifier), *DLCName.ToString() );
+			{
+				FScopeLock Lock(&PackageDataLock);
+				Package->OnDownloading(PackageIdentifier);
+			}
+		},
+		FGDKAsyncTaskQueue::GetBackgroundTaskQueue() // completion callback happens on a background thread
+	);
 
-		return false;
-	}
+	return SUCCEEDED(hResult);
 }
 
 #if !UE_BUILD_SHIPPING
@@ -1048,7 +1156,15 @@ bool FGDKPlatformDLC::Debug_GetPackageDescription( int PackageIndex, FString& Ou
 		Package = AllPackages[PackageIndex];
 	}
 
-	if (Package->MountState == EMountState::NotMounted || Package->MountState == EMountState::Unmounted)
+	if (Package->DownloadState == EDownloadState::NotDownloaded)
+	{
+		OutColor = FColor::Silver;
+	}
+	else if (Package->DownloadState == EDownloadState::Downloading)
+	{
+		OutColor = FColor::Cyan;
+	}
+	else if (Package->MountState == EMountState::NotMounted || Package->MountState == EMountState::Unmounted)
 	{
 		OutColor = FColor::Silver;
 	}
@@ -1066,7 +1182,18 @@ bool FGDKPlatformDLC::Debug_GetPackageDescription( int PackageIndex, FString& Ou
 		OutColor = FColor::Red;
 	}
 
-	OutDescription = FString::Printf(TEXT(".......%-30s %-30s %-30s (%s)"), *Package->DLCName.ToString(), *LexToString(Package->MountState), *LexToString(Package->LicenseState), *Package->PackageName );
+	FString DownloadState;
+	if (Package->DownloadState == EDownloadState::Downloading && Package->CachedTotalBytes > 0 && Package->CachedTotalBytes != UINT64_MAX)
+	{
+		const float Percent = 100.0f * (float)Package->CachedInstalledBytes / (float)Package->CachedTotalBytes;
+		DownloadState = FString::Printf(TEXT("Downloading(%.0f%%)"), Percent);
+	}
+	else
+	{
+		DownloadState = LexToString(Package->DownloadState);
+	}
+
+	OutDescription = FString::Printf(TEXT(".......%-30s %-20s %-15s %-15s (%s)"), *Package->DLCName.ToString(), *DownloadState, *LexToString(Package->MountState), *LexToString(Package->LicenseState), *Package->PackageName );
 	return true;
 }
 #endif //UE_BUILD_SHIPPING
@@ -1112,4 +1239,15 @@ FString LexToString( FGDKPlatformDLC::EProcessingState ProcessingState )
 	}
 }
 
+FString LexToString( FGDKPlatformDLC::EDownloadState DownloadState )
+{
+	switch(DownloadState)
+	{
+		case FGDKPlatformDLC::EDownloadState::NotDownloaded:  return TEXT("NotDownloaded");
+		case FGDKPlatformDLC::EDownloadState::Downloading:    return TEXT("Downloading");
+		case FGDKPlatformDLC::EDownloadState::Downloaded:     return TEXT("Downloaded");
+		default:                                              return TEXT("Invalid");
+	}
+}
+
 #endif //WITH_GRDK
diff --git a/Engine/Plugins/Experimental/Microsoft/GDKPlatformDLC/Source/GDKPlatformDLC/Private/GDKPlatformDLC.h b/Engine/Plugins/Experimental/Microsoft/GDKPlatformDLC/Source/GDKPlatformDLC/Private/GDKPlatformDLC.h
--- a/Engine/Plugins/Experimental/Microsoft/GDKPlatformDLC/Source/GDKPlatformDLC/Private/GDKPlatformDLC.h
+++ b/Engine/Plugins/Experimental/Microsoft/GDKPlatformDLC/Source/GDKPlatformDLC/Private/GDKPlatformDLC.h
@@ -75,25 +75,44 @@ private:
 		LicenseLost,
 		Unlicensed_DevIgnore,
 	};
+	enum class EDownloadState : uint8
+	{
+		NotDownloaded,
+		Downloading,
+		Downloaded,
+	};
 
 	struct FPackageData
 	{
-		FPackageData(FGDKPlatformDLC* InOwner, const XPackageDetails& InDetails, const FName& InDLCName);
+		FPackageData(FGDKPlatformDLC* InOwner, const FName& InDLCName, const FString& InStoreId);
 		~FPackageData();
 
-		char                        PackageIdentifier[XPACKAGE_IDENTIFIER_MAX_LENGTH];
-		char                        StoreId[STORE_SKU_ID_SIZE];
-		std::atomic<EMountState>    MountState;
-		std::atomic<ELicenseState>  LicenseState;
-		FString                     PackageName;
-		FString                     MountPath;
-		FGDKPlatformDLC*            Owner;
-		bool                        bIsAgeRestricted;
-		FName                       DLCName;
-		XPackageMountHandle         DLCHandle;
-		XStoreLicenseHandle         DLCLicense;
-		XTaskQueueRegistrationToken PackageLicenseLostToken;
-		TSharedPtr<class IPlugin>   DLCPlugin;
+		void OnDownloading(const char* InPackageIdentifier);
+		void OnDownloaded(const XPackageDetails& InDetails);
+		void OnUninstalled();
+
+		void CreateInstallationMonitor(const char* Identifier);
+		void DestroyInstallationMonitor();
+
+		char                              PackageIdentifier[XPACKAGE_IDENTIFIER_MAX_LENGTH] = {};
+		char                              StoreId[STORE_SKU_ID_SIZE] = {};
+		std::atomic<EMountState>          MountState = EMountState::NotMounted;
+		std::atomic<ELicenseState>        LicenseState = ELicenseState::Unlicensed;
+		std::atomic<EDownloadState>       DownloadState = EDownloadState::NotDownloaded;
+		FString                           PackageName;
+		FString                           MountPath;
+		FGDKPlatformDLC*                  Owner = nullptr;
+		bool                              bIsAgeRestricted = false;
+		bool                              bWasPreviouslyMounted = false;
+		FName                             DLCName;
+		XPackageMountHandle               DLCHandle = nullptr;
+		XStoreLicenseHandle               DLCLicense = nullptr;
+		XTaskQueueRegistrationToken       PackageLicenseLostToken = {};
+		TSharedPtr<class IPlugin>         DLCPlugin;
+		XPackageInstallationMonitorHandle InstallationMonitorHandle = nullptr;
+		XTaskQueueRegistrationToken       InstallationProgressToken = {};
+		std::atomic<uint64>               CachedInstalledBytes = 0;
+		std::atomic<uint64>               CachedTotalBytes = 0;
 
 		// context used for GDK callbacks that only accept void*
 		struct FContext
@@ -130,7 +149,7 @@ private:
 	TSharedPtr<FPackageData>         CurrentProcessingPackage;
 	FGDKAsyncTaskMonitor             CurrentProcessingAsyncTask;
 
-	void AddToPendingPackages( TSharedPtr<FPackageData> Package );
+	bool AddToPendingPackages( TSharedPtr<FPackageData> Package );
 	void InitializePackage_GameThread( TSharedPtr<FPackageData> Package );
 	void AsyncQueryLicense_GameThread( TSharedPtr<FPackageData> Package );
 	void AsyncAcquireLicense_GameThread( TSharedPtr<FPackageData> Package );
@@ -146,6 +165,7 @@ private:
 	FPlatformUserId     CurrentStoreUserId;
 	bool EnsureValidStoreContext();
 	static void OnPackageLicenseLost(void* Context);
+	static void OnInstallationProgressChanged(void* Context, XPackageInstallationMonitorHandle Monitor);
 
 	// configuration
 	bool bAutoMountDLC = true;            // Whether to automatically mount any new DLC that is downloaded once the application is up and running.
@@ -169,6 +189,7 @@ private:
 	// internal LexToString
 	friend FString LexToString( EMountState MountState );
 	friend FString LexToString( ELicenseState LicenseState );
+	friend FString LexToString( EDownloadState DownloadState );
 	friend FString LexToString( EProcessingState ProcessingState );
 };
 
