Remove workaround for re-mounting previously uninstalled DLC
- fix was to release the license when unmounting/uninstalling

#rnx
#Synced-CL 51936548

diff --git a/Engine/Plugins/Experimental/Microsoft/GDKPlatformDLC/Source/GDKPlatformDLC/Private/GDKPlatformDLC.cpp b/Engine/Plugins/Experimental/Microsoft/GDKPlatformDLC/Source/GDKPlatformDLC/Private/GDKPlatformDLC.cpp
--- a/Engine/Plugins/Experimental/Microsoft/GDKPlatformDLC/Source/GDKPlatformDLC/Private/GDKPlatformDLC.cpp
+++ b/Engine/Plugins/Experimental/Microsoft/GDKPlatformDLC/Source/GDKPlatformDLC/Private/GDKPlatformDLC.cpp
@@ -21,13 +21,6 @@
 #include "PlatformDLCModule.h"
 
 
-static bool GGDKDLCAllowRedownloadAfterMount = false;
-static FAutoConsoleVariableRef CVarDLCAllowRedownloadAfterMount(
-	TEXT( "GDK.DLC.AllowRedownloadAfterMount" ),
-	GGDKDLCAllowRedownloadAfterMount,
-	TEXT( "Whether to allow DLC to be downloaded again after it has previously been mounted. Disabled by default due to a bug that causes device I/O errors when re-mounting DLC that was previously uninstalled" ));
-
-
 FGDKPlatformDLC::FGDKPlatformDLC()
 	: bIsInitialized(false)
 	, CurrentStoreContext(nullptr)
@@ -238,12 +231,6 @@ void FGDKPlatformDLC::AsyncInitDLC(const XPackageDetails& Details, bool bMount)
 
 bool FGDKPlatformDLC::AddToPendingPackages( TSharedPtr<FPackageData> Package )
 {
-	if (Package->bWasPreviouslyMounted && !GGDKDLCAllowRedownloadAfterMount)
-	{
-		UE_LOG(LogPlatformDLC, Error, TEXT("DLC %s cannot currently be re-mounted after being previously uninstalled. Please restart the application to re-mount this DLC. Set GDK.DLC.AllowRedownloadAfterMount=1 to override."), *Package->DLCName.ToString());
-		return false;
-	}
-
 	FScopeLock Lock(&PackageDataLock);
 	PendingPackages.AddUnique(Package);
 	AddToTicker();
@@ -632,6 +619,15 @@ bool FGDKPlatformDLC::Unmount( FName DLCName )
 
 	if (Package->DLCHandle != nullptr)
 	{
+		// we must release the license before unmounting. if we hold this and the DLC is uninstalled & reinstalled, we may get device I/O errors
+		if (Package->DLCLicense != nullptr)
+		{
+			XStoreUnregisterPackageLicenseLost( Package->DLCLicense, Package->PackageLicenseLostToken, true);
+			XStoreCloseLicenseHandle(Package->DLCLicense);
+			Package->DLCLicense = nullptr;
+			Package->LicenseState = ELicenseState::Unlicensed;
+		}
+
 		// attempt to unregister the plugin as much as possible
 		if (bAutoMountPlugin && bAutoMountPakFiles)
 		{
@@ -1074,12 +1070,6 @@ bool FGDKPlatformDLC::Download( FName DLCName )
 		return true;
 	}
 
-	if (Package->bWasPreviouslyMounted && !GGDKDLCAllowRedownloadAfterMount)
-	{
-		UE_LOG(LogPlatformDLC, Error, TEXT("DLC %s cannot currently be re-downloaded after being mounted. Please restart the application to re-acquire this DLC. Set GDK.DLC.AllowRedownloadAfterMount=1 to override."), *DLCName.ToString());
-		return false;
-	}
-
 	// make sure there is a store context
 	if (!EnsureValidStoreContext())
 	{
