Some tidying of the GDK DLC handler:
- Remove QueryingLicense step as isn't necessary.
- Remove bAutoAcquireLicense ... no reason why someone would not want this, and it simplifies things.
- Remove bAllowDevOnlySideLoading... it isn't reliable. True sideloading still works, where the package is installed via dev tools instead of downloaded from the store.
- Remove bWasPreviouslyMounted... no longer used.
- Added a Mount_ prefix to the processing steps that are exclusively associated with mounting, for clarity.

other changes
- Mount() will fail early if the package isn't downloaded.
- Entitlement success notification sent earlier (will be moved again once purchasing is added)
- GDK.DLC.DebugShow and GDK.DLC.DumpState will load the DLC module automaticaly
- License is released when uninstalling, not unmounting - license remains valid during unmounting.

#Synced-CL 51936548

diff --git a/Engine/Plugins/Experimental/Microsoft/GDKPlatformDLC/Source/GDKPlatformDLC/Private/GDKPlatformDLC.cpp b/Engine/Plugins/Experimental/Microsoft/GDKPlatformDLC/Source/GDKPlatformDLC/Private/GDKPlatformDLC.cpp
--- a/Engine/Plugins/Experimental/Microsoft/GDKPlatformDLC/Source/GDKPlatformDLC/Private/GDKPlatformDLC.cpp
+++ b/Engine/Plugins/Experimental/Microsoft/GDKPlatformDLC/Source/GDKPlatformDLC/Private/GDKPlatformDLC.cpp
@@ -46,10 +46,6 @@ void FGDKPlatformDLC::Initialize()
 	GConfig->GetBool(ConfigSectionName, TEXT("AutoMountPakFiles"), bAutoMountPakFiles, GEngineIni);
 	GConfig->GetBool(ConfigSectionName, TEXT("AutoMountPlugin"), bAutoMountPlugin, GEngineIni);
 	GConfig->GetBool(ConfigSectionName, TEXT("AutoMountAgeRestricted"), bAutoMountAgeRestricted, GEngineIni);
-	GConfig->GetBool(ConfigSectionName, TEXT("AutoAcquireLicense"), bAutoAcquireLicense, GEngineIni);
-#if !UE_BUILD_SHIPPING
-	GConfig->GetBool(ConfigSectionName, TEXT("AllowDevOnlySideLoading"), bAllowDevOnlySideLoading, GEngineIni);
-#endif
 
 	// make sure the package manifest module is loaded
 	IGDKPackageManifestModule::Get();
@@ -170,7 +166,7 @@ bool FGDKPlatformDLC::Tick( float DeltaSeconds )
 			{
 				// get the next package to mount
 				CurrentProcessingPackage = PendingPackages.Pop();
-				InitializePackage_GameThread(CurrentProcessingPackage);
+				BeginProcessingPackage_GameThread(CurrentProcessingPackage);
 			}
 			else if (LostLicensePackages.Num() > 0)
 			{
@@ -188,30 +184,20 @@ bool FGDKPlatformDLC::Tick( float DeltaSeconds )
 		}
 		break;
 
-		case EProcessingState::QueryingLicense:
-		case EProcessingState::AcquiringLicense:
-		case EProcessingState::MountingPackage:
-		case EProcessingState::CachingData:
+		case EProcessingState::Mount_AcquiringLicense:
+		case EProcessingState::Mount_MountingPackage:
+		case EProcessingState::Mount_CachingData:
+		case EProcessingState::Mount_MountingPlugin:
 		{
 			// nothing to do, but sanity check that the package is still valid
 			check(CurrentProcessingPackage.IsValid());
 		}
 		break;
 
-		case EProcessingState::MountingPlugin:
+		case EProcessingState::Mount_Completed:
 		{
-			check(CurrentProcessingPackage.IsValid());
-			MountPlugin_GameThread(CurrentProcessingPackage);
-			check(ProcessingState == EProcessingState::Completed);
-		}
-		//... fall through to completion ...
-
-		case EProcessingState::Completed:
-		{
-			check(CurrentProcessingPackage.IsValid());
-			Complete_GameThread(CurrentProcessingPackage);
-			check(!CurrentProcessingPackage.IsValid());
-			check(ProcessingState == EProcessingState::Idle);
+			Mount_Completed_GameThread(CurrentProcessingPackage);
+			CurrentProcessingPackage.Reset();
 		}
 		break;
 	}
@@ -266,13 +252,13 @@ bool FGDKPlatformDLC::AddToPendingPackages( TSharedPtr<FPackageData> Package )
 }
 
 
-void FGDKPlatformDLC::InitializePackage_GameThread( TSharedPtr<FPackageData> Package )
+void FGDKPlatformDLC::BeginProcessingPackage_GameThread( TSharedPtr<FPackageData> Package )
 {
 	// if the package is already mounted and ready, just fire the callbacks again
-	bool bHasLicense = (Package->LicenseState == ELicenseState::Licensed || Package->LicenseState == ELicenseState::Unlicensed_DevIgnore);
+	bool bHasLicense = (Package->LicenseState == ELicenseState::Licensed);
 	if (Package->MountState == EMountState::Mounted && bHasLicense)
 	{
-		ProcessingState = EProcessingState::Completed;
+		ProcessingState = EProcessingState::Mount_Completed;
 		return;
 	}
 
@@ -287,90 +273,51 @@ void FGDKPlatformDLC::InitializePackage_GameThread( TSharedPtr<FPackageData> Pac
 	if (Package->bIsAgeRestricted && !bAutoMountAgeRestricted)
 	{
 		Package->MountState = EMountState::Ignored_AgeRestricted;
-		ProcessingState = EProcessingState::Completed;
+		ProcessingState = EProcessingState::Mount_Completed;
 		UE_LOG(LogPlatformDLC, Log, TEXT("Ignoring DLC %s because it is age restricted for one or more of the currently logged in users"), *Package->PackageName );
 		return;
 	}
 
-	// see if we should start acquiring a license
-	bool bWantLicense = false;
-	if (bAutoAcquireLicense)
-	{
-		// make sure there is a store context
-		if (!EnsureValidStoreContext())
-		{
-			ProcessingState = EProcessingState::Completed;
-			UE_LOG(LogPlatformDLC, Error, TEXT("cannot mount DLC"));
-			return;
-		}
+	// start the mounting process by getting a license
+	ProcessingState = EProcessingState::Mount_AcquiringLicense;
+	Mount_AsyncAcquireLicense_GameThread(Package);
+}
 
-		// see if we should acquire a license for this DLC
-		bWantLicense = (CurrentStoreContext != nullptr && Package->DLCLicense == nullptr);
-	}
 
-	if (bWantLicense)
-	{
-		ProcessingState = EProcessingState::QueryingLicense;
-		AsyncQueryLicense_GameThread(Package);
-	}
-	else
-	{
-		ProcessingState = EProcessingState::MountingPackage;
-		AsyncMountPackage_GameThread(Package);
-	}
-}
 
 
-void FGDKPlatformDLC::AsyncQueryLicense_GameThread( TSharedPtr<FPackageData> Package )
+void FGDKPlatformDLC::Mount_AsyncAcquireLicense_GameThread( TSharedPtr<FPackageData> Package )
 {
-	check(Package->DLCLicense == nullptr);
-	check(ProcessingState == EProcessingState::QueryingLicense);
+	check(ProcessingState == EProcessingState::Mount_AcquiringLicense);
 
-	HRESULT hResult = AsyncGDKTask(
-		CurrentProcessingAsyncTask,
-		[Package, this](XAsyncBlock* Block)
-		{
-			return XStoreCanAcquireLicenseForPackageAsync(CurrentStoreContext, Package->PackageIdentifier, Block);
-		},
-		[Package, this](XAsyncBlock* Block)
+	// if we already have a license, check whether it's still valid
+	if (Package->DLCLicense != nullptr)
+	{
+		if (XStoreIsLicenseValid(Package->DLCLicense))
 		{
-			// early out if the DLC system is being destroyed
-			bool bCancel = (ProcessingState == EProcessingState::Destroying);
-			if (bCancel)
-			{
-				return;
-			}
-
-			XStoreCanAcquireLicenseResult Result;
-			HRESULT hResult = XStoreCanAcquireLicenseForPackageResult(Block, &Result);
-			UE_CLOG(FAILED(hResult), LogPlatformDLC, Error, TEXT("failed to query license for package %s : 0x%X"), *Package->PackageName, hResult);
-
-			if (SUCCEEDED(hResult) && Result.status != XStoreCanLicenseStatus::NotLicensableToUser)
-			{
-				ProcessingState = EProcessingState::AcquiringLicense;
-				AsyncAcquireLicense_GameThread(Package);
-			}
-			else
-			{
-				Package->LicenseState = ELicenseState::Unlicensed;
-				ProcessingState = EProcessingState::Completed;
-				UE_CLOG(SUCCEEDED(hResult), LogPlatformDLC, Error, TEXT("Package %s is not licensable for the current store user (result = %d)"), *Package->PackageName, Result.status);
-			}
+			// license is still valid - skip straight to mounting
+			ProcessingState = EProcessingState::Mount_MountingPackage;
+			Mount_AsyncMountPackage_GameThread(Package);
+			return;
 		}
-	);
 
-	if (FAILED(hResult))
-	{
-		ProcessingState = EProcessingState::Completed;
-		UE_LOG(LogPlatformDLC, Error, TEXT("Failed to start querying licensability for package %s : 0x%X"), *Package->PackageName, hResult);
+		// license handle is stale - unregister the lost callback, close it, and re-acquire below
+		XStoreUnregisterPackageLicenseLost(Package->DLCLicense, Package->PackageLicenseLostToken, true);
+		XStoreCloseLicenseHandle(Package->DLCLicense);
+		Package->DLCLicense = nullptr;
+		Package->LicenseState = ELicenseState::Unlicensed;
+		PostNotification(Package->DLCName, ENotification::Entitlement, false);
 	}
-}
 
+	// make sure there is a store context
+	if (!EnsureValidStoreContext())
+	{
+		ProcessingState = EProcessingState::Mount_Completed;
+		UE_LOG(LogPlatformDLC, Error, TEXT("cannot mount DLC"));
+		return;
+	}
 
-void FGDKPlatformDLC::AsyncAcquireLicense_GameThread( TSharedPtr<FPackageData> Package )
-{
 	check(Package->DLCLicense == nullptr);
-	check(ProcessingState == EProcessingState::AcquiringLicense);
 
 	HRESULT hResult = AsyncGDKTask(
 		CurrentProcessingAsyncTask,
@@ -392,66 +339,49 @@ void FGDKPlatformDLC::AsyncAcquireLicense_GameThread( TSharedPtr<FPackageData> P
 			if (SUCCEEDED(hResult) && XStoreIsLicenseValid(Package->DLCLicense) )
 			{
 				Package->LicenseState = ELicenseState::Licensed;
+				PostNotification(Package->DLCName, ENotification::Entitlement, true); // @todo: this should be sent after a successful purchase (once that code is written), not here
 
 				// register for a callback when the license is lost
 				hResult = XStoreRegisterPackageLicenseLost(Package->DLCLicense, FGDKAsyncTaskQueue::GetGenericQueue(), &Package->Context, &FGDKPlatformDLC::OnPackageLicenseLost, &Package->PackageLicenseLostToken );
 				UE_CLOG(FAILED(hResult), LogPlatformDLC, Warning, TEXT("Could not register for GDK license lost callback : 0x%X"), hResult);
 
 				// start mounting the DLC package
-				ProcessingState = EProcessingState::MountingPackage;
-				AsyncMountPackage_GameThread(Package);
+				ProcessingState = EProcessingState::Mount_MountingPackage;
+				Mount_AsyncMountPackage_GameThread(Package);
 				return;
 			}
 
 			// if we received a license that was not valid, close it now as there's nothing we can do with it
 			if (Package->DLCLicense != nullptr)
 			{
+				XStoreUnregisterPackageLicenseLost(Package->DLCLicense, Package->PackageLicenseLostToken, true);
 				XStoreCloseLicenseHandle(Package->DLCLicense);
 				Package->DLCLicense = nullptr;
 			}
 
-			// handle failures
-#if !UE_BUILD_SHIPPING
-			if (bAllowDevOnlySideLoading && (
-				hResult == 0x87e10bc6 || // LM_E_CONTENT_NOT_IN_CATALOG
-				hResult == 0x803f9006 || // LM_E_ENTITLED_USER_SIGNED_OUT
-				hResult == 0x87dd0013))  // AM_E_USER_NOT_FOUND
-			{
-				// this is most likely from using the GDK development tools to sideload the DLC package instead of aquiring it from the Store. Allow the mount to happen regardless
-				UE_LOG(LogPlatformDLC, Warning, TEXT("Ignoring DLC license aquisition failure (0x%X) in non-shipping build. This will fail in Shipping"), hResult );
-				Package->LicenseState = ELicenseState::Unlicensed_DevIgnore;
-
-				ProcessingState = EProcessingState::MountingPackage;
-				AsyncMountPackage_GameThread(Package);
-			}
-			else
-#endif
-			{
-				Package->LicenseState = ELicenseState::Unlicensed;
-				ProcessingState = EProcessingState::Completed;
-				UE_LOG(LogPlatformDLC, Error, TEXT("Failed to get license for package %s : 0x%X"), *Package->PackageName, hResult);
-			}
+			Package->LicenseState = ELicenseState::Unlicensed;
+			ProcessingState = EProcessingState::Mount_Completed;
+			UE_LOG(LogPlatformDLC, Error, TEXT("Failed to get valid license for package %s : 0x%X"), *Package->PackageName, hResult);
 		}
 	);
 
 	if (FAILED(hResult))
 	{
 		Package->LicenseState = ELicenseState::Unlicensed;
-		ProcessingState = EProcessingState::Completed;
+		ProcessingState = EProcessingState::Mount_Completed;
 		UE_LOG(LogPlatformDLC, Error, TEXT("Failed to start getting license for package %s : 0x%X"), *Package->PackageName, hResult);
 	}
 }
 
 
-void FGDKPlatformDLC::AsyncMountPackage_GameThread( TSharedPtr<FPackageData> Package )
+void FGDKPlatformDLC::Mount_AsyncMountPackage_GameThread( TSharedPtr<FPackageData> Package )
 {
-	check(ProcessingState == EProcessingState::MountingPackage);
+	check(ProcessingState == EProcessingState::Mount_MountingPackage);
 
 	// if the package is already mounted, it means we've only reaquired the license
 	if (Package->DLCHandle != nullptr && ensure(Package->MountState == EMountState::Mounted) )
 	{
-		ProcessingState = EProcessingState::Completed;
-		PostNotification(Package->DLCName, ENotification::Entitlement, true);
+		ProcessingState = EProcessingState::Mount_Completed;
 		return;
 	}
 
@@ -476,16 +406,16 @@ void FGDKPlatformDLC::AsyncMountPackage_GameThread( TSharedPtr<FPackageData> Pac
 			if (SUCCEEDED(hResult))
 			{
 				// continue to initialize the DLC package on a background thread
-				ProcessingState = EProcessingState::CachingData;
+				ProcessingState = EProcessingState::Mount_CachingData;
 				AsyncTask( ENamedThreads::AnyBackgroundThreadNormalTask, [Package, this]()
 				{
-					CachePackageData_BackgroundThread(Package);
+					Mount_CachePackageData_BackgroundThread(Package);
 				});
 				return;
 			}
 
 			Package->MountState = EMountState::NotMounted;
-			ProcessingState = EProcessingState::Completed;
+			ProcessingState = EProcessingState::Mount_Completed;
 			UE_LOG(LogPlatformDLC, Error, TEXT("Failed to mount DLC '%s' : 0x%X"), *Package->PackageName, hResult);
 		}
 	);
@@ -493,15 +423,15 @@ void FGDKPlatformDLC::AsyncMountPackage_GameThread( TSharedPtr<FPackageData> Pac
 	if (FAILED(hResult))
 	{
 		Package->MountState = EMountState::NotMounted;
-		ProcessingState = EProcessingState::Completed;
+		ProcessingState = EProcessingState::Mount_Completed;
 		UE_LOG(LogPlatformDLC, Error, TEXT("Failed to start mounting mount DLC '%s' : 0x%X"), *Package->PackageName, hResult);
 	}
 }
 
 
-void FGDKPlatformDLC::CachePackageData_BackgroundThread( TSharedPtr<FPackageData> Package )
+void FGDKPlatformDLC::Mount_CachePackageData_BackgroundThread( TSharedPtr<FPackageData> Package )
 {
-	check(ProcessingState == EProcessingState::CachingData);
+	check(ProcessingState == EProcessingState::Mount_CachingData);
 
 	HRESULT hResult;
 
@@ -511,7 +441,7 @@ void FGDKPlatformDLC::CachePackageData_BackgroundThread( TSharedPtr<FPackageData
 	if (FAILED(hResult))
 	{
 		Package->MountState = EMountState::NotMounted;
-		ProcessingState = EProcessingState::Completed;
+		ProcessingState = EProcessingState::Mount_Completed;
 		UE_LOG(LogPlatformDLC, Error, TEXT("Failed to get mount path for DLC '%s' : 0x%X"), *Package->PackageName, hResult);
 
 		XPackageCloseMountHandle(Package->DLCHandle);
@@ -525,13 +455,17 @@ void FGDKPlatformDLC::CachePackageData_BackgroundThread( TSharedPtr<FPackageData
 	IGDKPackageManifestModule::Get().LoadManifestFromDLC(Package->DLCHandle, Package->PackageIdentifier, Package->MountPath);
 
 	// finalize on the GameThread
-	ProcessingState = EProcessingState::MountingPlugin;
+	ProcessingState = EProcessingState::Mount_MountingPlugin;
+	ExecuteOnGameThread( UE_SOURCE_LOCATION, [this, Package]()
+	{
+		Mount_MountPlugin_GameThread(Package);
+	});
 }
 
 
-void FGDKPlatformDLC::MountPlugin_GameThread( TSharedPtr<FPackageData> Package )
+void FGDKPlatformDLC::Mount_MountPlugin_GameThread( TSharedPtr<FPackageData> Package )
 {
-	check(ProcessingState == EProcessingState::MountingPlugin);
+	check(ProcessingState == EProcessingState::Mount_MountingPlugin);
 
 	if (bAutoMountPakFiles)
 	{
@@ -601,21 +535,16 @@ void FGDKPlatformDLC::MountPlugin_GameThread( TSharedPtr<FPackageData> Package )
 
 	// all done
 	Package->MountState = EMountState::Mounted;
-	Package->bWasPreviouslyMounted = true;
-	ProcessingState = EProcessingState::Completed;
+	ProcessingState = EProcessingState::Mount_Completed;
 }
 
 
-void FGDKPlatformDLC::Complete_GameThread( TSharedPtr<FPackageData> Package )
+void FGDKPlatformDLC::Mount_Completed_GameThread( TSharedPtr<FPackageData> Package )
 {
-	check(IsInGameThread());
-
 	// signal callback
 	bool bSuccess = (Package->MountState == EMountState::Mounted);
 	PostNotification(Package->DLCName, ENotification::Mounted, bSuccess);
 
-	// remove the current work item & check for more work next tick
-	CurrentProcessingPackage.Reset();
 	ProcessingState = EProcessingState::Idle;
 }
 
@@ -653,15 +582,6 @@ bool FGDKPlatformDLC::Unmount( FName DLCName )
 
 	if (Package->DLCHandle != nullptr)
 	{
-		// we must release the license before unmounting. if we hold this and the DLC is uninstalled & reinstalled, we may get device I/O errors
-		if (Package->DLCLicense != nullptr)
-		{
-			XStoreUnregisterPackageLicenseLost( Package->DLCLicense, Package->PackageLicenseLostToken, true);
-			XStoreCloseLicenseHandle(Package->DLCLicense);
-			Package->DLCLicense = nullptr;
-			Package->LicenseState = ELicenseState::Unlicensed;
-		}
-
 		// attempt to unregister the plugin as much as possible
 		if (bAutoMountPlugin && bAutoMountPakFiles && Package->DLCPlugin.IsValid())
 		{
@@ -843,7 +763,7 @@ bool FGDKPlatformDLC::IsInitialized() const
 bool FGDKPlatformDLC::HasEntitlement( FName DLCName ) const
 {
 	TSharedPtr<const FPackageData> Package = FindDLC(DLCName);
-	return Package.IsValid() && (Package->LicenseState == ELicenseState::Licensed || Package->LicenseState == ELicenseState::Unlicensed_DevIgnore);
+	return Package.IsValid() && (Package->LicenseState == ELicenseState::Licensed);
 }
 
 bool FGDKPlatformDLC::IsMounted( FName DLCName ) const
@@ -867,6 +787,12 @@ bool FGDKPlatformDLC::Mount( FName DLCName )
 		return false;
 	}
 
+	if (Package->DownloadState != EDownloadState::Downloaded)
+	{
+		UE_LOG(LogPlatformDLC, Error, TEXT("DLC %s must be downloaded before it can be mounted"), *DLCName.ToString());
+		return false;
+	}
+
 	return AddToPendingPackages(Package);
 }
 
@@ -946,6 +872,7 @@ FGDKPlatformDLC::FOnDLCNotification& FGDKPlatformDLC::OnNotification()
 
 void FGDKPlatformDLC::PostNotification(FName DLCName, ENotification Notification, bool bSuccess)
 {
+	UE_LOG(LogPlatformDLC, Log, TEXT("PostNotification %s: %s %s"), *DLCName.ToString(), *LexToString(Notification), *LexToString(bSuccess));
 	ExecuteOnGameThread( UE_SOURCE_LOCATION, [this, DLCName, Notification, bSuccess]
 	{
 		OnNotificationDelegate.Broadcast(DLCName, Notification, bSuccess);
@@ -1072,6 +999,16 @@ bool FGDKPlatformDLC::Uninstall( FName DLCName )
 	{
 		if (Package->DLCHandle == nullptr)
 		{
+			// we must release the license before uninstaliling. if we hold this and the DLC is uninstalled & reinstalled, we may get device I/O errors
+			if (Package->DLCLicense != nullptr)
+			{
+				XStoreUnregisterPackageLicenseLost( Package->DLCLicense, Package->PackageLicenseLostToken, true);
+				XStoreCloseLicenseHandle(Package->DLCLicense);
+				Package->DLCLicense = nullptr;
+				Package->LicenseState = ELicenseState::Unlicensed;
+				PostNotification(Package->DLCName, ENotification::Entitlement, false);
+			}
+
 			bool bSuccess = XPackageUninstallPackage(Package->PackageIdentifier);
 			if (bSuccess)
 			{
@@ -1201,7 +1138,7 @@ bool FGDKPlatformDLC::Debug_GetPackageDescription( int PackageIndex, FString& Ou
 	}
 	else if (Package->MountState == EMountState::Mounted)
 	{
-		bool bIsLicensed = (Package->LicenseState == ELicenseState::Licensed) || (Package->LicenseState == ELicenseState::Unlicensed_DevIgnore);
+		bool bIsLicensed = (Package->LicenseState == ELicenseState::Licensed);
 		OutColor = bIsLicensed ? FColor::Green : FColor::Yellow;
 	}
 	else if (Package->MountState < EMountState::Mounted)
@@ -1249,7 +1186,6 @@ FString LexToString( FGDKPlatformDLC::ELicenseState LicenseState )
 		case FGDKPlatformDLC::ELicenseState::Unlicensed:            return TEXT("Unlicensed");
 		case FGDKPlatformDLC::ELicenseState::Licensed:              return TEXT("Licensed");
 		case FGDKPlatformDLC::ELicenseState::LicenseLost:           return TEXT("LicenseLost");
-		case FGDKPlatformDLC::ELicenseState::Unlicensed_DevIgnore:  return TEXT("Unlicensed(Sideload)");
         default:                                                    return TEXT("Invalid");
 	}
 }
@@ -1258,15 +1194,14 @@ FString LexToString( FGDKPlatformDLC::EProcessingState ProcessingState )
 {
 	switch(ProcessingState)
 	{
-        case FGDKPlatformDLC::EProcessingState::Idle:               return TEXT("Idle");
-		case FGDKPlatformDLC::EProcessingState::QueryingLicense:    return TEXT("QueryingLicense");
-		case FGDKPlatformDLC::EProcessingState::AcquiringLicense:   return TEXT("AcquiringLicense");
-		case FGDKPlatformDLC::EProcessingState::MountingPackage:    return TEXT("MountingPackage");
-		case FGDKPlatformDLC::EProcessingState::CachingData:        return TEXT("CachingData");
-		case FGDKPlatformDLC::EProcessingState::MountingPlugin:     return TEXT("MountingPlugin");
-        case FGDKPlatformDLC::EProcessingState::Completed:          return TEXT("Completed");
-		case FGDKPlatformDLC::EProcessingState::Destroying:         return TEXT("Destroying");
-        default:                                                    return TEXT("Invalid");
+        case FGDKPlatformDLC::EProcessingState::Idle:                     return TEXT("Idle");
+		case FGDKPlatformDLC::EProcessingState::Mount_AcquiringLicense:   return TEXT("Mounting(AcquiringLicense)");
+		case FGDKPlatformDLC::EProcessingState::Mount_MountingPackage:    return TEXT("Mounting(MountingPackage)");
+		case FGDKPlatformDLC::EProcessingState::Mount_CachingData:        return TEXT("Mounting(CachingData)");
+		case FGDKPlatformDLC::EProcessingState::Mount_MountingPlugin:     return TEXT("Mounting(MountingPlugin)");
+        case FGDKPlatformDLC::EProcessingState::Mount_Completed:          return TEXT("Mounting(Completed)");
+		case FGDKPlatformDLC::EProcessingState::Destroying:               return TEXT("Destroying");
+        default:                                                          return TEXT("Invalid");
 	}
 }
 
diff --git a/Engine/Plugins/Experimental/Microsoft/GDKPlatformDLC/Source/GDKPlatformDLC/Private/GDKPlatformDLC.h b/Engine/Plugins/Experimental/Microsoft/GDKPlatformDLC/Source/GDKPlatformDLC/Private/GDKPlatformDLC.h
--- a/Engine/Plugins/Experimental/Microsoft/GDKPlatformDLC/Source/GDKPlatformDLC/Private/GDKPlatformDLC.h
+++ b/Engine/Plugins/Experimental/Microsoft/GDKPlatformDLC/Source/GDKPlatformDLC/Private/GDKPlatformDLC.h
@@ -74,7 +74,6 @@ private:
 		Unlicensed,
 		Licensed,
 		LicenseLost,
-		Unlicensed_DevIgnore,
 	};
 	enum class EDownloadState : uint8
 	{
@@ -104,7 +103,6 @@ private:
 		FString                           MountPath;
 		FGDKPlatformDLC*                  Owner = nullptr;
 		bool                              bIsAgeRestricted = false;
-		bool                              bWasPreviouslyMounted = false;
 		FName                             DLCName;
 		XPackageMountHandle               DLCHandle = nullptr;
 		XStoreLicenseHandle               DLCLicense = nullptr;
@@ -138,12 +136,11 @@ private:
 	enum class EProcessingState : uint8
 	{
 		Idle,
-		QueryingLicense,
-		AcquiringLicense,
-		MountingPackage,
-		CachingData,
-		MountingPlugin,
-		Completed,
+		Mount_AcquiringLicense,
+		Mount_MountingPackage,
+		Mount_CachingData,
+		Mount_MountingPlugin,
+		Mount_Completed,
 		Destroying,
 	};
 	std::atomic<EProcessingState>    ProcessingState = EProcessingState::Idle;
@@ -153,13 +150,12 @@ private:
 	FGDKAsyncTaskMonitor             CurrentProcessingAsyncTask;
 
 	bool AddToPendingPackages( TSharedPtr<FPackageData> Package );
-	void InitializePackage_GameThread( TSharedPtr<FPackageData> Package );
-	void AsyncQueryLicense_GameThread( TSharedPtr<FPackageData> Package );
-	void AsyncAcquireLicense_GameThread( TSharedPtr<FPackageData> Package );
-	void AsyncMountPackage_GameThread( TSharedPtr<FPackageData> Package );
-	void CachePackageData_BackgroundThread( TSharedPtr<FPackageData> Package );
-	void MountPlugin_GameThread( TSharedPtr<FPackageData> Package );
-	void Complete_GameThread( TSharedPtr<FPackageData> Package );
+	void BeginProcessingPackage_GameThread( TSharedPtr<FPackageData> Package );
+	void Mount_AsyncAcquireLicense_GameThread( TSharedPtr<FPackageData> Package );
+	void Mount_AsyncMountPackage_GameThread( TSharedPtr<FPackageData> Package );
+	void Mount_CachePackageData_BackgroundThread( TSharedPtr<FPackageData> Package );
+	void Mount_MountPlugin_GameThread( TSharedPtr<FPackageData> Package );
+	void Mount_Completed_GameThread( TSharedPtr<FPackageData> Package );
 	void LicenseLost_GameThread( TSharedPtr<FPackageData> Package );
 
 	// store
@@ -176,12 +172,6 @@ private:
 	bool bAutoMountPakFiles = true;       // Whether to automatically mount pak files
 	bool bAutoMountPlugin = true;         // Whether to automatically mount the DLC plugin and asset registry etc. (requires bAutoMountPakFiles)
 	bool bAutoMountAgeRestricted = false; // Whether to attempt to mount age-restricted DLC. This will trigger a parental controls consent prompt.
-	bool bAutoAcquireLicense = true;      // Whether to attempt to acquire the license for any DLC before attempting to mount. If this is false, the title must manage licensing directly.
-#if !UE_BUILD_SHIPPING
-	bool bAllowDevOnlySideLoading = true; // Whether to ignore certain failure states when license aquisition fails in a non-shipping build - typically during dev sideloading scenarios.
-#endif
-
-
 
 	// ticker
 	FTSTicker::FDelegateHandle TickHandle;
diff --git a/Engine/Plugins/Experimental/Microsoft/GDKPlatformDLC/Source/GDKPlatformDLC/Private/GDKPlatformDLCModule.cpp b/Engine/Plugins/Experimental/Microsoft/GDKPlatformDLC/Source/GDKPlatformDLC/Private/GDKPlatformDLCModule.cpp
--- a/Engine/Plugins/Experimental/Microsoft/GDKPlatformDLC/Source/GDKPlatformDLC/Private/GDKPlatformDLCModule.cpp
+++ b/Engine/Plugins/Experimental/Microsoft/GDKPlatformDLC/Source/GDKPlatformDLC/Private/GDKPlatformDLCModule.cpp
@@ -73,6 +73,10 @@ public:
 						GEngine->AddOnScreenDebugMessage(INDEX_NONE, 0.0f, PackageColor, PackageDescription, false);
 					}
 				}
+				else
+				{
+					GEngine->AddOnScreenDebugMessage(INDEX_NONE, 0.0f, FColor::White, TEXT("GDK DLC handler has not been created"), false);
+				}
 				return true;
 			});
 		}
@@ -119,7 +123,7 @@ static FAutoConsoleCommand CVarGDKDumpDLC(
 	TEXT("Dumps the state of the DLC"),
 	FConsoleCommandDelegate::CreateLambda([]()
 	{
-		FGDKPlatformDLCModule* Module = FModuleManager::GetModulePtr<FGDKPlatformDLCModule>("GDKPlatformDLC");
+		FGDKPlatformDLCModule* Module = FModuleManager::LoadModulePtr<FGDKPlatformDLCModule>("GDKPlatformDLC");
 		if (Module != nullptr && Module->DLCHandler.IsValid())
 		{
 			UE_LOG( LogPlatformDLC, Log, TEXT("%s"), *Module->DLCHandler->Debug_GetStateDescription() );
@@ -141,8 +145,8 @@ static FAutoConsoleCommand CVarGDKDLCShow(
 	TEXT("Show or hide the on-screen debugging"),
 	FConsoleCommandWithArgsDelegate::CreateLambda( [](const TArray<FString>& Args )
 	{
-		FGDKPlatformDLCModule* Module = FModuleManager::GetModulePtr<FGDKPlatformDLCModule>("GDKPlatformDLC");
-		if (Module != nullptr && Module->DLCHandler.IsValid())
+		FGDKPlatformDLCModule* Module = FModuleManager::LoadModulePtr<FGDKPlatformDLCModule>("GDKPlatformDLC");
+		if (Module != nullptr)
 		{
 			bool bShow = (Args.Num() == 0) || (FCString::Atoi(*Args[0]) != 0);
 			Module->SetOnScreenDebugRender(bShow);
